[Insights]
* Added possibility to specify port number when start listening for a direct trace.
* Added command line parameters to start Unreal Insights in "listening for a direct trace" mode: -ListenForDirectTrace [-DirectTracePort=<port>]
* Session Info: Added "Stream Type" info line ("Trace File" / "Trace Store" / "Direct Trace" / "Custom Stream").
* Improved API re starting a trace analysis. Separated SetAutoQuit() from StartAnalysis*() in IUnrealInsightsModule.

diff --git a/Engine/Source/Developer/TraceAnalysis/Private/DataStream.cpp b/Engine/Source/Developer/TraceAnalysis/Private/DataStream.cpp
--- a/Engine/Source/Developer/TraceAnalysis/Private/DataStream.cpp
+++ b/Engine/Source/Developer/TraceAnalysis/Private/DataStream.cpp
@@ -178,15 +178,20 @@ FDirectSocketStream::~FDirectSocketStream()
 }
 	
 
-uint16 FDirectSocketStream::StartListening()
+uint16 FDirectSocketStream::StartListening(uint16 Port)
 {
-	uint16 Port = DefaultPort;
+	if (Port == 0)
+	{
+		Port = DefaultPort;
+	}
 	// Try to bind the default port. If that is busy move to next candidate
-	while (!CreateSocket(Port) && (Port - DefaultPort) < MaxPortAttempts)
+	int32 Attempts = 0;
+	while (!CreateSocket(Port) && Attempts < MaxPortAttempts)
 	{
 		++Port;
+		++Attempts;
 	}
-	if (Port - DefaultPort >= MaxPortAttempts)
+	if (Attempts >= MaxPortAttempts)
 	{
 		return 0;
 	}
diff --git a/Engine/Source/Developer/TraceAnalysis/Public/Trace/DataStream.h b/Engine/Source/Developer/TraceAnalysis/Public/Trace/DataStream.h
--- a/Engine/Source/Developer/TraceAnalysis/Public/Trace/DataStream.h
+++ b/Engine/Source/Developer/TraceAnalysis/Public/Trace/DataStream.h
@@ -17,7 +17,7 @@ namespace Trace {
 class IInDataStream
 {
 public:
-	virtual			~IInDataStream() = default;
+	virtual ~IInDataStream() = default;
 
 	/**
 	 * Read bytes from the stream.
@@ -25,22 +25,22 @@ public:
 	 * @param Size Maximum size that can be read in bytes
 	 * @return Number of bytes read from the stream. Zero indicates end of stream and negative values indicate errors.
 	 */
-	virtual int32	Read(void* Data, uint32 Size) = 0;
+	virtual int32 Read(void* Data, uint32 Size) = 0;
 
 	/**
 	 * Close the stream. Reading from a closed stream
 	 * is considered an error.
 	 */
-	virtual void	Close() {}
-	
+	virtual void Close() {}
+
 	/**
 	 * Query if the stream is ready to read. Some streams may need to
 	 * establish the data stream before reading can begin. A stream may not
 	 * block indefinitely.
-	 * 
+	 *
 	 * @return if the stream is ready to be read from
 	 */
-	virtual bool	WaitUntilReady() { return true; }
+	virtual bool WaitUntilReady() { return true; }
 };
 
 /*
@@ -54,9 +54,9 @@ public:
 
 	/*
 	* Open the file.
-	* 
-	* @param Path	The path to the file.
-	* 
+	*
+	* @param Path The path to the file.
+	*
 	* @return True if the file was opened successfully.
 	*/
 	UE_API bool Open(const TCHAR* Path);
@@ -82,9 +82,10 @@ public:
 	/**
 	 * Initiates listening sockets. Must be called before attempting to read from
 	 * the stream.
-	 * @return Port number used for listening
+	 * @param Port The specific port number to connect to or 0 to use any available port
+	 * @return The actual port number used for listening
 	 */
-	UE_API uint16 StartListening();
+	UE_API uint16 StartListening(uint16 Port = 0);
 
 private:
 	enum
@@ -102,7 +103,7 @@ private:
 	// FRunnable interface
 	UE_API virtual uint32 Run() override;
 	UE_API virtual void Stop() override;
-	
+
 	void Accept();
 	bool CreateSocket(uint16 Port);
 
diff --git a/Engine/Source/Developer/TraceInsights/Private/Insights/InsightsManager.cpp b/Engine/Source/Developer/TraceInsights/Private/Insights/InsightsManager.cpp
--- a/Engine/Source/Developer/TraceInsights/Private/Insights/InsightsManager.cpp
+++ b/Engine/Source/Developer/TraceInsights/Private/Insights/InsightsManager.cpp
@@ -471,8 +471,10 @@ void FInsightsManager::ResetSession(bool bNotify)
 		Session->Stop(true);
 		Session.Reset();
 
+		CurrentTraceStreamType = ETraceStreamType::Unknown;
+		CurrentTraceName.Reset();
 		CurrentTraceId = 0;
-		CurrentTraceFilename.Reset();
+		CurrentDirectTracePort = 0;
 
 		if (bNotify)
 		{
@@ -500,7 +502,7 @@ void FInsightsManager::PollAnalysisInfo()
 			TSharedPtr<TraceServices::IAnalysisSession> EditableSession = ConstCastSharedPtr<TraceServices::IAnalysisSession>(Session);
 			TraceServices::FAnalysisSessionEditScope SessionEditScope(*EditableSession.Get());
 			const auto Messages = EditableSession->DrainPendingMessages();
-			
+
 			FMessageLog ReportMessageLog(AnalysisLogListingName);
 			for (const auto& Message : Messages)
 			{
@@ -523,7 +525,7 @@ void FInsightsManager::OnSessionChanged()
 	{
 		FMessageLog(AnalysisLogListingName).NewPage(FText::FromString(Session->GetName()));
 	}
-	
+
 	SessionChangedEvent.Broadcast();
 }
 
@@ -691,7 +693,7 @@ void FInsightsManager::AutoLoadLiveSession()
 
 ////////////////////////////////////////////////////////////////////////////////////////////////////
 
-void FInsightsManager::LoadLastLiveSession(float RetryTime)
+void FInsightsManager::LoadLastLiveSession(float InRetryTime)
 {
 	ResetSession();
 
@@ -720,22 +722,22 @@ void FInsightsManager::LoadLastLiveSession(float RetryTime)
 		LoadTrace(LastLiveSessionTraceId);
 	}
 
-	if (Session == nullptr && RetryTime > 0)
+	if (Session == nullptr && InRetryTime > 0)
 	{
-		RetryLoadLastLiveSessionTimer = RetryTime;
+		RetryLoadLastLiveSessionTimer = InRetryTime;
 	}
 }
 
 ////////////////////////////////////////////////////////////////////////////////////////////////////
 
-void FInsightsManager::LoadTrace(uint32 InTraceId, bool InAutoQuit, bool InWaitForSymbolResolver)
+void FInsightsManager::LoadTrace(uint32 InTraceId)
 {
 	ResetSession();
 
 	UE::Trace::FStoreClient* StoreClient = TraceStoreConnection.GetStoreClient();
 	if (!StoreClient)
 	{
-		if (InAutoQuit)
+		if (bAutoQuit)
 		{
 			RequestEngineExit(AutoQuitMsgOnFail);
 		}
@@ -747,7 +749,7 @@ void FInsightsManager::LoadTrace(uint32 InTraceId, bool InAutoQuit, bool InWaitF
 	UE::Trace::FStoreClient::FTraceData TraceData = StoreClient->ReadTrace(InTraceId);
 	if (!TraceData)
 	{
-		if (InAutoQuit)
+		if (bAutoQuit)
 		{
 			RequestEngineExit(AutoQuitMsgOnFail);
 		}
@@ -780,14 +782,14 @@ void FInsightsManager::LoadTrace(uint32 InTraceId, bool InAutoQuit, bool InWaitF
 
 	if (Session)
 	{
+		CurrentTraceStreamType = ETraceStreamType::TraceStore;
+		CurrentTraceName = TraceName;
 		CurrentTraceId = InTraceId;
-		CurrentTraceFilename = TraceName;
+		CurrentDirectTracePort = 0;
 		bIsSessionInfoSet = false;
 		OnSessionChanged();
-		bSessionAnalysisCompletedAutoQuit = InAutoQuit;
-		bWaitForSymbolResolver = InAutoQuit && InWaitForSymbolResolver;
 	}
-	else if (InAutoQuit)
+	else if (bAutoQuit)
 	{
 		RequestEngineExit(AutoQuitMsgOnFail);
 	}
@@ -806,13 +808,13 @@ void FInsightsManager::LoadTraceFile()
 
 ////////////////////////////////////////////////////////////////////////////////////////////////////
 
-void FInsightsManager::LoadTraceFile(const FString& InTraceFilename, bool InAutoQuit, bool InWaitForSymbolResolver)
+void FInsightsManager::LoadTraceFile(const FString& InTraceFilename)
 {
 	IPlatformFile& PlatformFile = IPlatformFile::GetPlatformPhysical();
 	if (!PlatformFile.FileExists(*InTraceFilename))
 	{
 		const uint32 TraceId = uint32(FCString::Strtoui64(*InTraceFilename, nullptr, 10));
-		return LoadTrace(TraceId, InAutoQuit);
+		return LoadTrace(TraceId);
 	}
 
 	ResetSession();
@@ -821,14 +823,14 @@ void FInsightsManager::LoadTraceFile(const FString& InTraceFilename, bool InAuto
 
 	if (Session)
 	{
+		CurrentTraceStreamType = ETraceStreamType::TraceFile;
+		CurrentTraceName = InTraceFilename;
 		CurrentTraceId = 0;
-		CurrentTraceFilename = InTraceFilename;
+		CurrentDirectTracePort = 0;
 		bIsSessionInfoSet = false;
 		OnSessionChanged();
-		bSessionAnalysisCompletedAutoQuit = InAutoQuit;
-		bWaitForSymbolResolver = InAutoQuit && InWaitForSymbolResolver;
 	}
-	else if (InAutoQuit)
+	else if (bAutoQuit)
 	{
 		RequestEngineExit(AutoQuitMsgOnFail);
 	}
@@ -836,44 +838,57 @@ void FInsightsManager::LoadTraceFile(const FString& InTraceFilename, bool InAuto
 
 ////////////////////////////////////////////////////////////////////////////////////////////////////
 
-uint16 FInsightsManager::ListenForDirectTrace(const TCHAR* InStreamName)
+uint16 FInsightsManager::ListenForDirectTrace(const TCHAR* InStreamName, uint16 InPort)
 {
-	constexpr uint32 TraceId = 0; 
-	FString TraceName = InStreamName ? InStreamName : TEXT("DirectStream");
-	TUniquePtr<UE::Trace::FDirectSocketStream> TraceData = MakeUnique<UE::Trace::FDirectSocketStream>();
-	const uint16 Port = TraceData->StartListening();
-	
 	ResetSession();
 
+	TUniquePtr<UE::Trace::FDirectSocketStream> TraceData = MakeUnique<UE::Trace::FDirectSocketStream>();
+	uint16 Port = TraceData->StartListening(InPort);
+
+	constexpr uint32 TraceId = 0;
+	FString TraceName = InStreamName ? InStreamName : FString::Printf(TEXT("Direct Trace Stream on Port %u"), uint32(Port));
 	Session = AnalysisService->StartAnalysis(TraceId, *TraceName, MoveTemp(TraceData));
 
 	if (Session)
 	{
-		CurrentTraceId = TraceId;
-		CurrentTraceFilename = TraceName;
+		CurrentTraceStreamType = ETraceStreamType::DirectTrace;
+		CurrentTraceName = TraceName;
+		CurrentTraceId = 0;
+		CurrentDirectTracePort = Port;
 		bIsSessionInfoSet = false;
 		OnSessionChanged();
 	}
+	else if (bAutoQuit)
+	{
+		RequestEngineExit(AutoQuitMsgOnFail);
+	}
 
 	return Port;
 }
 
 ////////////////////////////////////////////////////////////////////////////////////////////////////
+
 void FInsightsManager::StartAnalysisWithStream(TUniquePtr<UE::Trace::IInDataStream>&& InStream, const TCHAR* InStreamName)
 {
-	ResetSession(false);
+	ResetSession();
 
-	constexpr uint32 TraceId = 0; 
+	constexpr uint32 TraceId = 0;
 	FString TraceName = InStreamName ? InStreamName : TEXT("CustomStream");
 	Session = AnalysisService->StartAnalysis(TraceId, *TraceName, MoveTemp(InStream));
 
 	if (Session)
 	{
-		CurrentTraceId = TraceId;
-		CurrentTraceFilename = TraceName;
+		CurrentTraceStreamType = ETraceStreamType::Custom;
+		CurrentTraceName = TraceName;
+		CurrentTraceId = 0;
+		CurrentDirectTracePort = 0;
 		bIsSessionInfoSet = false;
 		OnSessionChanged();
 	}
+	else if (bAutoQuit)
+	{
+		RequestEngineExit(AutoQuitMsgOnFail);
+	}
 }
 
 ////////////////////////////////////////////////////////////////////////////////////////////////////
@@ -931,17 +946,23 @@ bool FInsightsManager::OnDrop(const FDragDropEvent& DragDropEvent)
 
 void FInsightsManager::UpdateAppTitle()
 {
-#if !WITH_EDITOR
+#if !WITH_EDITOR || defined(__INTELLISENSE__)
 	TSharedPtr<SWindow> RootWindow = FGlobalTabmanager::Get()->GetRootWindow();
 	if (RootWindow)
 	{
-		if (CurrentTraceFilename.IsEmpty())
+		if (CurrentTraceName.IsEmpty())
 		{
 			const FText AppTitle = LOCTEXT("UnrealInsightsAppName", "Unreal Insights");
 			RootWindow->SetTitle(AppTitle);
 		}
 		else
 		{
+			const FString SessionName =
+				(CurrentTraceStreamType == ETraceStreamType::TraceFile ||
+				 CurrentTraceStreamType == ETraceStreamType::TraceStore) ?
+					FPaths::GetBaseFilename(CurrentTraceName) :
+					CurrentTraceName;
+
 			bool bWasAppTitleUpdated = false;
 			if (Session.IsValid())
 			{
@@ -951,7 +972,6 @@ void FInsightsManager::UpdateAppTitle()
 				{
 					TraceServices::FSessionInfo SessionInfo = DiagnosticsProvider->GetSessionInfo();
 
-					const FString SessionName = FPaths::GetBaseFilename(CurrentTraceFilename);
 					const FText AppTitle = FText::Format(LOCTEXT("UnrealInsightsAppNameFmt2", "{0}{1} - {2} - {3} - {4} - {5} Unreal Insights"),
 						FText::FromString(SessionName),
 						!SessionInfo.Branch.IsEmpty() ? FText::FromString(TEXT(" - ") + SessionInfo.Branch) : FText::GetEmpty(),
@@ -967,7 +987,6 @@ void FInsightsManager::UpdateAppTitle()
 
 			if (!bWasAppTitleUpdated)
 			{
-				const FString SessionName = FPaths::GetBaseFilename(CurrentTraceFilename);
 				const FText AppTitle = FText::Format(LOCTEXT("UnrealInsightsAppNameFmt", "{0} - Unreal Insights"), FText::FromString(SessionName));
 				RootWindow->SetTitle(AppTitle);
 			}
@@ -1009,13 +1028,13 @@ void FInsightsManager::OnSessionAnalysisCompleted()
 	if (FInsightsTestRunner::Get().IsValid())
 	{
 		// Don't quit now. Let the test runner to execute.
-		bSessionAnalysisCompletedAutoQuit = false;
+		bAutoQuit = false;
 	}
 #endif
 
-	if (bSessionAnalysisCompletedAutoQuit && !bWaitForSymbolResolver)
+	if (bAutoQuit && !bWaitForSymbolResolver)
 	{
-		bSessionAnalysisCompletedAutoQuit = false;
+		bAutoQuit = false;
 		RequestEngineExit(AutoQuitMsg);
 	}
 }
@@ -1024,8 +1043,9 @@ void FInsightsManager::OnSessionAnalysisCompleted()
 
 void FInsightsManager::OnSessionSymbolResolverCompleted()
 {
-	if (bWaitForSymbolResolver)
+	if (bAutoQuit && bWaitForSymbolResolver)
 	{
+		bAutoQuit = false;
 		bWaitForSymbolResolver = false;
 		RequestEngineExit(AutoQuitMsg);
 	}
diff --git a/Engine/Source/Developer/TraceInsights/Private/Insights/InsightsManager.h b/Engine/Source/Developer/TraceInsights/Private/Insights/InsightsManager.h
--- a/Engine/Source/Developer/TraceInsights/Private/Insights/InsightsManager.h
+++ b/Engine/Source/Developer/TraceInsights/Private/Insights/InsightsManager.h
@@ -57,6 +57,15 @@ struct FAsyncTaskData
 	}
 };
 
+enum class ETraceStreamType
+{
+	Unknown,
+	TraceFile,
+	TraceStore,
+	DirectTrace,
+	Custom
+};
+
 ////////////////////////////////////////////////////////////////////////////////////////////////////
 /**
  * This class manages following areas:
@@ -124,11 +133,17 @@ public:
 	/** @return an instance of the trace analysis session. */
 	TSharedPtr<const TraceServices::IAnalysisSession> GetSession() const;
 
-	/** @return the id of the trace being analyzed. */
+	/** @return the stream type of the trace being analyzed. */
+	ETraceStreamType GetTraceStreamType() const { return CurrentTraceStreamType; }
+
+	/** @return the name of the trace being analyzed. */
+	const FString& GetTraceName() const { return CurrentTraceName; }
+
+	/** @return the store id of the trace being analyzed. */
 	uint32 GetTraceId() const { return CurrentTraceId; }
 
-	/** @return the filename of the trace being analyzed. */
-	const FString& GetTraceFilename() const { return CurrentTraceFilename; }
+	/** @return the port of the direct trace connection. */
+	uint16 GetDirectTracePort() const { return CurrentDirectTracePort; }
 
 	//////////////////////////////////////////////////
 
@@ -170,8 +185,16 @@ public:
 	// Getters and setters used by Toggle Commands.
 
 	/** @return true, if UI is allowed to display debug info. */
-	const bool IsDebugInfoEnabled() const { return bIsDebugInfoEnabled; }
-	void SetDebugInfo(const bool bDebugInfoEnabledState) { bIsDebugInfoEnabled = bDebugInfoEnabledState; }
+	bool IsDebugInfoEnabled() const { return bIsDebugInfoEnabled; }
+	void SetDebugInfo(bool bDebugInfoEnabledState) { bIsDebugInfoEnabled = bDebugInfoEnabledState; }
+
+	bool IsAutoQuitEnabled() const { return bAutoQuit; }
+	bool IsWaitForSymbolResolverEnabled() const { return bAutoQuit && bWaitForSymbolResolver; }
+	void SetAutoQuit(bool bInAutoQuit, bool bInWaitForSymbolResolver = false)
+	{
+		bAutoQuit = bInAutoQuit;
+		bWaitForSymbolResolver = bInWaitForSymbolResolver;
+	}
 
 	////////////////////////////////////////////////////////////////////////////////////////////////////
 
@@ -207,11 +230,9 @@ public:
 	/**
 	 * Creates a new analysis session instance using the specified trace id.
 	 * Replaces the current analysis session.
-	 * @param TraceId - The id of the trace to analyze
-	 * @param InAutoQuit - The Application will close when session analysis is complete or fails to start
-	 * @param InWaitForSymbolResolver - The Application will automatically close when both session analysis is complete and symbols have been resolved. Must be set together with the InAutoQuit parameter
+	 * @param InTraceId - The id of the trace to analyze
 	 */
-	void LoadTrace(uint32 TraceId, bool InAutoQuit = false, bool InWaitForSymbolResolver = false);
+	void LoadTrace(uint32 InTraceId);
 
 	/**
 	 * Shows the open file dialog and creates a new analysis session instance for the chosen trace file.
@@ -222,22 +243,22 @@ public:
 	/**
 	 * Creates a new analysis session instance and loads a trace file from the specified location.
 	 * Replaces the current analysis session.
-	 * @param TraceFilename - The trace file to analyze
-	 * @param InAutoQuit - The Application will close when session analysis is complete or fails to start
-	 * @param InWaitForSymbolResolver - The Application will automatically close when both session analysis is complete and symbols have been resolved. Must be set together with the InAutoQuit parameter
+	 * @param InTraceFilename - The trace file to analyze
 	 */
-	void LoadTraceFile(const FString& TraceFilename, bool InAutoQuit = false, bool InWaitForSymbolResolver = false);
+	void LoadTraceFile(const FString& InTraceFilename);
 
 	/**
 	 * Creates a new analysis session instance and opens a listening port to await connections.
+	 * Replaces the current analysis session.
 	 * @param InStreamName User facing name of stream
-	 * @return Port number to connect to
-	 * @return Port used for listening to incoming connections. 0 if failed to open port.
+	 * @param InPort The specific port number to listen to or 0 to use any available port
+	 * @return The actual port used for listening to incoming connections. 0 if failed to open port.
 	 */
-	uint16 ListenForDirectTrace(const TCHAR* InStreamName = nullptr);
+	uint16 ListenForDirectTrace(const TCHAR* InStreamName = nullptr, uint16 InPort = 0);
 
 	/**
-	 * Creates a new analysis session instance from a provided data stream. Replaces the current analysis session.
+	 * Creates a new analysis session instance from a provided data stream.
+	 * Replaces the current analysis session.
 	 * @param InStream Stream to read from
 	 * @param InStreamName User facing name of stream
 	 */
@@ -351,11 +372,17 @@ private:
 	/** The trace analysis session. */
 	TSharedPtr<const TraceServices::IAnalysisSession> Session;
 
+	/** The stream type of the trace being analyzed. */
+	ETraceStreamType CurrentTraceStreamType = ETraceStreamType::Unknown;
+
+	/** The name of the trace being analyzed. */
+	FString CurrentTraceName;
+
 	/** The id of the trace being analyzed. */
 	uint32 CurrentTraceId = 0;
 
-	/** The filename of the trace being analyzed. */
-	FString CurrentTraceFilename;
+	/** The port of the direct trace connection. */
+	uint16 CurrentDirectTracePort = 0;
 
 	/** List of UI commands for this manager. This will be filled by this and corresponding classes. */
 	TSharedRef<FUICommandList> CommandList;
@@ -376,8 +403,7 @@ private:
 	bool bIsSessionInfoSet = false;
 
 	bool bIsAnalysisComplete = false;
-	bool bSessionAnalysisCompletedAutoQuit = false;
-
+	bool bAutoQuit = false;
 	bool bIsSymbolResolvingComplete = false;
 	bool bWaitForSymbolResolver = false;
 
diff --git a/Engine/Source/Developer/TraceInsights/Private/Insights/TraceInsightsModule.cpp b/Engine/Source/Developer/TraceInsights/Private/Insights/TraceInsightsModule.cpp
--- a/Engine/Source/Developer/TraceInsights/Private/Insights/TraceInsightsModule.cpp
+++ b/Engine/Source/Developer/TraceInsights/Private/Insights/TraceInsightsModule.cpp
@@ -482,11 +482,18 @@ TSharedPtr<const TraceServices::IAnalysisSession> FTraceInsightsModule::GetAnaly
 
 ////////////////////////////////////////////////////////////////////////////////////////////////////
 
-void FTraceInsightsModule::StartAnalysisForTrace(uint32 InTraceId, bool InAutoQuit, bool InWaitForSymbolResolver)
+void FTraceInsightsModule::SetAutoQuit(bool bInAutoQuit, bool bInWaitForSymbolResolver)
+{
+	FInsightsManager::Get()->SetAutoQuit(bInAutoQuit, bInWaitForSymbolResolver);
+}
+
+////////////////////////////////////////////////////////////////////////////////////////////////////
+
+void FTraceInsightsModule::StartAnalysisForTrace(uint32 InTraceId)
 {
 	if (InTraceId != 0)
 	{
-		FInsightsManager::Get()->LoadTrace(InTraceId, InAutoQuit, InWaitForSymbolResolver);
+		FInsightsManager::Get()->LoadTrace(InTraceId);
 	}
 	UpdateAppTitle();
 }
@@ -501,20 +508,22 @@ void FTraceInsightsModule::StartAnalysisForLastLiveSession(float InRetryTime)
 
 ////////////////////////////////////////////////////////////////////////////////////////////////////
 
-void FTraceInsightsModule::StartAnalysisForTraceFile(const TCHAR* InTraceFile, bool InAutoQuit, bool InWaitForSymbolResolver)
+void FTraceInsightsModule::StartAnalysisForTraceFile(const TCHAR* InTraceFile)
 {
 	if (InTraceFile != nullptr)
 	{
-		FInsightsManager::Get()->LoadTraceFile(FString(InTraceFile), InAutoQuit, InWaitForSymbolResolver);
+		FInsightsManager::Get()->LoadTraceFile(FString(InTraceFile));
 	}
 	UpdateAppTitle();
 }
 
 ////////////////////////////////////////////////////////////////////////////////////////////////////
 
-uint16 FTraceInsightsModule::StartAnalysisWithDirectTrace(const TCHAR* InStreamName)
+uint16 FTraceInsightsModule::StartAnalysisWithDirectTrace(const TCHAR* InStreamName, uint16 InPort)
 {
-	return FInsightsManager::Get()->ListenForDirectTrace(InStreamName);
+	uint16 Port = FInsightsManager::Get()->ListenForDirectTrace(InStreamName, InPort);
+	UpdateAppTitle();
+	return Port;
 }
 
 ////////////////////////////////////////////////////////////////////////////////////////////////////
@@ -522,6 +531,7 @@ uint16 FTraceInsightsModule::StartAnalysisWithDirectTrace(const TCHAR* InStreamN
 void FTraceInsightsModule::StartAnalysisWithStream(TUniquePtr<UE::Trace::IInDataStream>&& InStream, const TCHAR* InStreamName)
 {
 	FInsightsManager::Get()->StartAnalysisWithStream(MoveTemp(InStream), InStreamName);
+	UpdateAppTitle();
 }
 
 ////////////////////////////////////////////////////////////////////////////////////////////////////
@@ -538,13 +548,13 @@ void FTraceInsightsModule::OnWindowClosedEvent(const TSharedRef<SWindow>&)
 
 ////////////////////////////////////////////////////////////////////////////////////////////////////
 
-void FTraceInsightsModule::InitializeTesting(bool InInitAutomationModules, bool InAutoQuit)
+void FTraceInsightsModule::InitializeTesting(bool bInInitAutomationModules, bool bInAutoQuit)
 {
 #if !UE_BUILD_SHIPPING && !WITH_EDITOR
 	auto TestRunner = FInsightsTestRunner::CreateInstance();
 
-	TestRunner->SetInitAutomationModules(InInitAutomationModules);
-	TestRunner->SetAutoQuit(InAutoQuit);
+	TestRunner->SetInitAutomationModules(bInInitAutomationModules);
+	TestRunner->SetAutoQuit(bInAutoQuit);
 
 	RegisterComponent(TestRunner);
 #endif
diff --git a/Engine/Source/Developer/TraceInsights/Private/Insights/TraceInsightsModule.h b/Engine/Source/Developer/TraceInsights/Private/Insights/TraceInsightsModule.h
--- a/Engine/Source/Developer/TraceInsights/Private/Insights/TraceInsightsModule.h
+++ b/Engine/Source/Developer/TraceInsights/Private/Insights/TraceInsightsModule.h
@@ -58,11 +58,15 @@ public:
 
 	virtual void CreateSessionViewer(bool bAllowDebugTools = false) override;
 
+	virtual TSharedPtr<::Insights::IInsightsManager> GetInsightsManager() override;
 	virtual TSharedPtr<const TraceServices::IAnalysisSession> GetAnalysisSession() const override;
-	virtual void StartAnalysisForTrace(uint32 InTraceId, bool InAutoQuit = false, bool InWaitForSymbolResolver = false) override;
+
+	virtual void SetAutoQuit(bool bInAutoQuit, bool bInWaitForSymbolResolver = false) override;
+
+	virtual void StartAnalysisForTrace(uint32 InTraceId) override;
 	virtual void StartAnalysisForLastLiveSession(float InRetryTime = 1.0f) override;
-	virtual void StartAnalysisForTraceFile(const TCHAR* InTraceFile, bool InAutoQuit = false, bool InWaitForSymbolResolver = false) override;
-	virtual uint16 StartAnalysisWithDirectTrace(const TCHAR* InStreamName = nullptr) override;
+	virtual void StartAnalysisForTraceFile(const TCHAR* InTraceFile) override;
+	virtual uint16 StartAnalysisWithDirectTrace(const TCHAR* InStreamName = nullptr, uint16 InPort = 0) override;
 	virtual void StartAnalysisWithStream(TUniquePtr<UE::Trace::IInDataStream>&& InStream, const TCHAR* InStreamName = nullptr) override;
 
 	virtual void ShutdownUserInterface() override;
@@ -70,8 +74,6 @@ public:
 	virtual void RegisterComponent(TSharedPtr<IInsightsComponent> Component) override;
 	virtual void UnregisterComponent(TSharedPtr<IInsightsComponent> Component) override;
 
-	virtual TSharedPtr<::Insights::IInsightsManager> GetInsightsManager() override;
-
 	virtual void RegisterMajorTabConfig(const FName& InMajorTabId, const FInsightsMajorTabConfig& InConfig) override;
 	virtual void UnregisterMajorTabConfig(const FName& InMajorTabId) override;
 	virtual FOnInsightsMajorTabCreated& OnMajorTabCreated() override { return OnInsightsMajorTabCreatedDelegate; }
diff --git a/Engine/Source/Developer/TraceInsights/Private/Insights/Widgets/SSessionInfoWindow.cpp b/Engine/Source/Developer/TraceInsights/Private/Insights/Widgets/SSessionInfoWindow.cpp
--- a/Engine/Source/Developer/TraceInsights/Private/Insights/Widgets/SSessionInfoWindow.cpp
+++ b/Engine/Source/Developer/TraceInsights/Private/Insights/Widgets/SSessionInfoWindow.cpp
@@ -295,7 +295,7 @@ TSharedRef<SDockTab> SSessionInfoWindow::SpawnTab_SessionInfo(const FSpawnTabArg
 	BeginSection(VerticalBox, LOCTEXT("SessionInfo_SectionText", "Session Info"));
 
 	AddInfoLine(VerticalBox, LOCTEXT("SessionName_HeaderText",		"Session Name"),		&SSessionInfoWindow::GetSessionNameText,	&SSessionInfoWindow::IsAlwaysVisible);
-	AddInfoLine(VerticalBox, LOCTEXT("Uri_HeaderText",				"URI"),					&SSessionInfoWindow::GetUriText,			&SSessionInfoWindow::IsAlwaysVisible);
+	AddInfoLine(VerticalBox, LOCTEXT("Uri_HeaderText",				"URI"),					&SSessionInfoWindow::GetUriText,			&SSessionInfoWindow::IsVisibleUriText);
 	AddInfoLine(VerticalBox, LOCTEXT("Platform_HeaderText",			"Platform"),			&SSessionInfoWindow::GetPlatformText,		&SSessionInfoWindow::IsVisiblePlatformText);
 	AddInfoLine(VerticalBox, LOCTEXT("AppName_HeaderText",			"Application Name"),	&SSessionInfoWindow::GetAppNameText,		&SSessionInfoWindow::IsVisibleAppNameText);
 	AddInfoLine(VerticalBox, LOCTEXT("ProjectName_HeaderText",		"Project Name"),		&SSessionInfoWindow::GetProjectNameText,	&SSessionInfoWindow::IsVisibleProjectNameText);
@@ -306,6 +306,7 @@ TSharedRef<SDockTab> SSessionInfoWindow::SpawnTab_SessionInfo(const FSpawnTabArg
 	AddInfoLine(VerticalBox, LOCTEXT("BuildTarget_HeaderText",		"Build Target"),		&SSessionInfoWindow::GetBuildTargetText,	&SSessionInfoWindow::IsVisibleBuildTargetText);
 	AddInfoLine(VerticalBox, LOCTEXT("CommandLine_HeaderText",		"Command Line"),		&SSessionInfoWindow::GetCommandLineText,	&SSessionInfoWindow::IsVisibleCommandLineText, true);
 	AddInfoLine(VerticalBox, LOCTEXT("OtherMetadata_HeaderText",	"Other Metadata"),		&SSessionInfoWindow::GetOtherMetadataText,	&SSessionInfoWindow::IsVisibleOtherMetadataText, true);
+	AddInfoLine(VerticalBox, LOCTEXT("StreamType_HeaderText",		"Stream Type"),			&SSessionInfoWindow::GetStreamTypeText,		&SSessionInfoWindow::IsAlwaysVisible);
 
 	EndSection(VerticalBox);
 
@@ -356,15 +357,51 @@ void SSessionInfoWindow::Tick(const FGeometry& AllottedGeometry, const double In
 		AnalysisSession = Session;
 		bIsSessionInfoSet = false;
 
+		ETraceStreamType TraceStreamType = FInsightsManager::Get()->GetTraceStreamType();
+
+		switch (TraceStreamType)
+		{
+			case ETraceStreamType::TraceFile:
+				StreamTypeText = LOCTEXT("TraceStreamType_TraceFile", "Trace File");
+				break;
+
+			case ETraceStreamType::TraceStore:
+				StreamTypeText = FText::Format(LOCTEXT("TraceStreamType_TraceStore_Fmt", "Trace Store {0} -- TraceId={1}"),
+					FText::FromString(FString::Printf(TEXT("%s:%u"), *FInsightsManager::Get()->GetLastStoreHost(), FInsightsManager::Get()->GetLastStorePort())),
+					FText::FromString(FString::Printf(TEXT("0x%X"), Session->GetTraceId())));
+				break;
+
+			case ETraceStreamType::DirectTrace:
+				StreamTypeText = FText::Format(LOCTEXT("TraceStreamType_DirectTrace_Fmt", "Direct Trace -- Port={0}"),
+					FText::AsNumber(uint32(FInsightsManager::Get()->GetDirectTracePort()), &FNumberFormattingOptions::DefaultNoGrouping()));
+				break;
+
+			case ETraceStreamType::Custom:
+				StreamTypeText = LOCTEXT("TraceStreamType_Custom", "Custom Stream");
+				break;
+
+			default:
+				StreamTypeText = LOCTEXT("TraceStreamType_Unknown", "Unknown");
+				break;
+		}
+
 		// We can quickly get the session name and uri.
 		{
 			TraceServices::FAnalysisSessionReadScope SessionReadScope(*Session.Get());
 
-			SessionNameText = FText::FromString(FPaths::GetBaseFilename(Session->GetName()));
+			if (TraceStreamType == ETraceStreamType::TraceFile ||
+				TraceStreamType == ETraceStreamType::TraceStore)
+			{
+				SessionNameText = FText::FromString(FPaths::GetBaseFilename(Session->GetName()));
 
-			FString Uri(Session->GetName());
-			FPaths::NormalizeFilename(Uri);
-			UriText = FText::FromString(Uri);
+				FString Uri(Session->GetName());
+				FPaths::NormalizeFilename(Uri);
+				UriText = FText::FromString(Uri);
+			}
+			else
+			{
+				SessionNameText = FText::FromString(Session->GetName());
+			}
 		}
 
 		PlatformText = FText::GetEmpty();
@@ -401,11 +438,6 @@ void SSessionInfoWindow::Tick(const FGeometry& AllottedGeometry, const double In
 
 			TStringBuilder<1024> OtherMetadata;
 
-			if (Session->GetTraceId() != 0)
-			{
-				OtherMetadata.Appendf(TEXT("TraceId=0x%X"), Session->GetTraceId());
-			}
-
 			Session->EnumerateMetadata([&OtherMetadata](const TraceServices::FTraceSessionMetadata& Metadata)
 				{
 					static FName ExcludedMetadata[]
diff --git a/Engine/Source/Developer/TraceInsights/Private/Insights/Widgets/SSessionInfoWindow.h b/Engine/Source/Developer/TraceInsights/Private/Insights/Widgets/SSessionInfoWindow.h
--- a/Engine/Source/Developer/TraceInsights/Private/Insights/Widgets/SSessionInfoWindow.h
+++ b/Engine/Source/Developer/TraceInsights/Private/Insights/Widgets/SSessionInfoWindow.h
@@ -137,6 +137,7 @@ private:
 	TSharedRef<SDockTab> SpawnTab_SessionInfo(const FSpawnTabArgs& Args);
 	void OnSessionInfoTabClosed(TSharedRef<SDockTab> TabBeingClosed);
 
+	FText GetStreamTypeText() const { return StreamTypeText; }
 	FText GetSessionNameText() const { return SessionNameText; }
 	FText GetUriText() const { return UriText; }
 
@@ -153,6 +154,7 @@ private:
 
 	EVisibility IsAlwaysVisible() const { return EVisibility::Visible; }
 
+	EVisibility IsVisibleUriText() const { return UriText.IsEmpty() ? EVisibility::Collapsed : EVisibility::Visible; }
 	EVisibility IsVisiblePlatformText() const { return PlatformText.IsEmpty() ? EVisibility::Collapsed : EVisibility::Visible; }
 	EVisibility IsVisibleAppNameText() const { return AppNameText.IsEmpty() ? EVisibility::Collapsed : EVisibility::Visible; }
 	EVisibility IsVisibleProjectNameText() const { return ProjectNameText.IsEmpty() ? EVisibility::Collapsed : EVisibility::Visible; }
@@ -180,6 +182,7 @@ private:
 	/** Holds the tab manager that manages the front-end's tabs. */
 	TSharedPtr<FTabManager> TabManager;
 
+	FText StreamTypeText;
 	FText SessionNameText;
 	FText UriText;
 
diff --git a/Engine/Source/Developer/TraceInsights/Public/Insights/IUnrealInsightsModule.h b/Engine/Source/Developer/TraceInsights/Public/Insights/IUnrealInsightsModule.h
--- a/Engine/Source/Developer/TraceInsights/Public/Insights/IUnrealInsightsModule.h
+++ b/Engine/Source/Developer/TraceInsights/Public/Insights/IUnrealInsightsModule.h
@@ -253,14 +253,27 @@ public:
 	 */
 	virtual TSharedPtr<const TraceServices::IAnalysisSession> GetAnalysisSession() const = 0;
 
+	/**
+	 * Sets the AutoQuit behavior when analysis is completed.
+	 *
+	 * @param bInAutoQuit - The Application will close when session analysis is complete or fails to start
+	 * @param bInWaitForSymbolResolver - The Application will automatically close when both session analysis is complete and symbols have been resolved.Must be set together with the InAutoQuit parameter.
+	 */
+	virtual void SetAutoQuit(bool bInAutoQuit, bool bInWaitForSymbolResolver = false) = 0;
+
 	/**
 	 * Starts analysis of the specified trace. Called when the application starts in "Viewer" mode.
 	 *
 	 * @param InTraceId The id of the trace to analyze.
-	 * @param InAutoQuit - The Application will close when session analysis is complete or fails to start
-	 * @param InWaitForSymbolResolver - The Application will automatically close when both session analysis is complete and symbols have been resolved. Must be set together with the InAutoQuit parameter.
 	 */
-	virtual void StartAnalysisForTrace(uint32 InTraceId, bool InAutoQuit = false, bool InWaitForSymbolResolver = false) = 0;
+	virtual void StartAnalysisForTrace(uint32 InTraceId) = 0;
+
+	UE_DEPRECATED(5.7, "Use explicit SetAutoQuit()")
+	virtual void StartAnalysisForTrace(uint32 InTraceId, bool bInAutoQuit, bool bInWaitForSymbolResolver = false)
+	{
+		SetAutoQuit(bInAutoQuit, bInWaitForSymbolResolver);
+		StartAnalysisForTrace(InTraceId);
+	}
 
 	/**
 	 * Starts analysis of the last live session. Called when the application starts in "Viewer" mode.
@@ -274,20 +287,25 @@ public:
 	 * Starts analysis of the specified *.utrace file. Called when the application starts in "Viewer" mode.
 	 *
 	 * @param InTraceFile The filename (*.utrace) of the trace to analyze.
-	 * @param InAutoQuit - The Application will close when session analysis is complete or fails to start
-	 * @param InWaitForSymbolResolver - The Application will automatically close when both session analysis is complete and symbols have been resolved. Must be set together with the InAutoQuit parameter.
 	 */
-	virtual void StartAnalysisForTraceFile(const TCHAR* InTraceFile, bool InAutoQuit = false, bool InWaitForSymbolResolver = false) = 0;
+	virtual void StartAnalysisForTraceFile(const TCHAR* InTraceFile) = 0;
+
+	UE_DEPRECATED(5.7, "Use explicit SetAutoQuit()")
+	virtual void StartAnalysisForTraceFile(const TCHAR* InTraceFile, bool bInAutoQuit, bool bInWaitForSymbolResolver = false)
+	{
+		SetAutoQuit(bInAutoQuit, bInWaitForSymbolResolver);
+		StartAnalysisForTraceFile(InTraceFile);
+	}
 
 	/**
 	 * Starts analysis for a direct connection from the tracing application (might be self). Notice that
 	 * analysis doesn't actually start until a connection has been made.
-	 * 
+	 *
 	 * @param InStreamName User facing display name for this stream
-	 * @return Port used for listening to incoming connections. 0 if failed to open port.
+	 * @param InPort The specific port number to listen to or 0 to use any available port
+	 * @return The actual port used for listening to incoming connections. 0 if failed to open port.
 	 */
-	virtual uint16 StartAnalysisWithDirectTrace(const TCHAR* InStreamName = nullptr) = 0;
-
+	virtual uint16 StartAnalysisWithDirectTrace(const TCHAR* InStreamName = nullptr, uint16 InPort = 0) = 0;
 
 	/**
 	 * Start analysis given a data stream.
diff --git a/Engine/Source/Programs/UnrealInsights/Private/UnrealInsightsMain.cpp b/Engine/Source/Programs/UnrealInsights/Private/UnrealInsightsMain.cpp
--- a/Engine/Source/Programs/UnrealInsights/Private/UnrealInsightsMain.cpp
+++ b/Engine/Source/Programs/UnrealInsights/Private/UnrealInsightsMain.cpp
@@ -120,7 +120,8 @@ int32 UnrealInsightsMain(const TCHAR* CommandLine)
 	bool bFrontendMode;
 	FString TraceFileToOpen;
 
-	if (FCString::Strifind(CommandLine, TEXT("-OpenTraceId=")))
+	if (FCString::Strifind(CommandLine, TEXT("-OpenTraceId=")) ||
+		FCString::Strifind(CommandLine, TEXT("-ListenForDirectTrace")))
 	{
 		bFrontendMode = false;
 	}
diff --git a/Engine/Source/Programs/UnrealInsights/Private/UserInterfaceCommand.cpp b/Engine/Source/Programs/UnrealInsights/Private/UserInterfaceCommand.cpp
--- a/Engine/Source/Programs/UnrealInsights/Private/UserInterfaceCommand.cpp
+++ b/Engine/Source/Programs/UnrealInsights/Private/UserInterfaceCommand.cpp
@@ -108,7 +108,7 @@ void FUserInterfaceCommand::Run(bool bFrontendMode, const FString& TraceFileToOp
 
 	//////////////////////////////////////////////////
 	// Initialize source code access.
-	
+
 	// Load the source code access module.
 	ISourceCodeAccessModule& SourceCodeAccessModule = FModuleManager::LoadModuleChecked<ISourceCodeAccessModule>(FName("SourceCodeAccess"));
 
@@ -249,6 +249,13 @@ void FUserInterfaceCommand::InitializeSlateApplication(bool bFrontendMode, const
 			}
 		}
 
+		const bool bListenForDirectTrace = FParse::Param(FCommandLine::Get(), TEXT("ListenForDirectTrace"));
+		uint16 DirectTracePort = 0;
+		if (bListenForDirectTrace)
+		{
+			FParse::Value(FCommandLine::Get(), TEXT("-DirectTracePort="), DirectTracePort);
+		}
+
 		FString Cmd;
 		if (FParse::Value(FCommandLine::Get(), TEXT("-ExecOnAnalysisCompleteCmd="), Cmd, false))
 		{
@@ -261,14 +268,19 @@ void FUserInterfaceCommand::InitializeSlateApplication(bool bFrontendMode, const
 			TraceInsightsModule.CreateSessionViewer(bAllowDebugTools);
 		}
 
-		if (bUseTraceId)
+		TraceInsightsModule.SetAutoQuit(bAutoQuit, bWaitForSymbolResolver);
+		if (bListenForDirectTrace)
+		{
+			TraceInsightsModule.StartAnalysisWithDirectTrace(nullptr, DirectTracePort);
+		}
+		else if (bUseTraceId)
 		{
 			TraceInsightsModule.ConnectToStore(*StoreHost, StorePort);
-			TraceInsightsModule.StartAnalysisForTrace(TraceId, bAutoQuit, bWaitForSymbolResolver);
+			TraceInsightsModule.StartAnalysisForTrace(TraceId);
 		}
 		else
 		{
-			TraceInsightsModule.StartAnalysisForTraceFile(*TraceFileToOpen, bAutoQuit, bWaitForSymbolResolver);
+			TraceInsightsModule.StartAnalysisForTraceFile(*TraceFileToOpen);
 		}
 	}
 	else // frontend mode
