[Insights] Timing View / Quick Find: Added support for searching by region name in Quick Find

diff --git a/Engine/Source/Developer/TraceInsights/Private/Insights/TimingProfiler/Tracks/RegionsTimingTrack.cpp b/Engine/Source/Developer/TraceInsights/Private/Insights/TimingProfiler/Tracks/RegionsTimingTrack.cpp
--- a/Engine/Source/Developer/TraceInsights/Private/Insights/TimingProfiler/Tracks/RegionsTimingTrack.cpp
+++ b/Engine/Source/Developer/TraceInsights/Private/Insights/TimingProfiler/Tracks/RegionsTimingTrack.cpp
@@ -189,6 +189,7 @@ void FTimingRegionsTrack::BuildFilteredDrawState(ITimingEventsTrackDrawStateBuil
 		FilterContext.AddFilterData<double>(static_cast<int32>(EFilterField::EndTime), 0.0f);
 		FilterContext.AddFilterData<double>(static_cast<int32>(EFilterField::Duration), 0.0f);
 		FilterContext.AddFilterData<FString>(static_cast<int32>(EFilterField::TrackName), this->GetName());
+		FilterContext.AddFilterData<FString>(static_cast<int32>(EFilterField::RegionName), "");
 
 		TSharedPtr<const TraceServices::IAnalysisSession> Session = FInsightsManager::Get()->GetSession();
 
@@ -209,6 +210,7 @@ void FTimingRegionsTrack::BuildFilteredDrawState(ITimingEventsTrackDrawStateBuil
 							FilterContext.SetFilterData<double>(static_cast<int32>(EFilterField::StartTime), Region.BeginTime);
 							FilterContext.SetFilterData<double>(static_cast<int32>(EFilterField::EndTime), Region.EndTime);
 							FilterContext.SetFilterData<double>(static_cast<int32>(EFilterField::Duration), Region.EndTime - Region.BeginTime);
+							FilterContext.SetFilterData<FString>(static_cast<int32>(EFilterField::RegionName), Region.Text);
 
 							RegionHadEvents = true;
 							if (FilterConfigurator->ApplyFilters(FilterContext))
@@ -261,7 +263,7 @@ bool FTimingRegionsTrack::FindRegionEvent(const FTimingEventSearchParameters& In
 	FilterConfiguratorContext.AddFilterData<double>(static_cast<int32>(EFilterField::EndTime), 0.0f);
 	FilterConfiguratorContext.AddFilterData<double>(static_cast<int32>(EFilterField::Duration), 0.0f);
 	FilterConfiguratorContext.AddFilterData<FString>(static_cast<int32>(EFilterField::TrackName), this->GetName());
-	FilterConfiguratorContext.AddFilterData<int64>(static_cast<int32>(EFilterField::RegionName), 0);
+	FilterConfiguratorContext.AddFilterData<FString>(static_cast<int32>(EFilterField::RegionName), "");
 
 	return TTimingEventSearch<TraceServices::FTimeRegion>::Search(
 	InParameters,
@@ -300,7 +302,7 @@ bool FTimingRegionsTrack::FindRegionEvent(const FTimingEventSearchParameters& In
 			FilterConfiguratorContext.SetFilterData<double>(static_cast<int32>(EFilterField::StartTime), EventStartTime);
 			FilterConfiguratorContext.SetFilterData<double>(static_cast<int32>(EFilterField::EndTime), EventEndTime);
 			FilterConfiguratorContext.SetFilterData<double>(static_cast<int32>(EFilterField::Duration), EventEndTime - EventStartTime);
-			FilterConfiguratorContext.SetFilterData<int64>(static_cast<int32>(EFilterField::RegionName), reinterpret_cast<int64>(Region.Text));
+			FilterConfiguratorContext.SetFilterData<FString>(static_cast<int32>(EFilterField::RegionName), Region.Text);
 			return InParameters.FilterExecutor->ApplyFilters(FilterConfiguratorContext);
 		}
 
diff --git a/Engine/Source/Developer/TraceInsights/Private/Insights/Widgets/STimingView.cpp b/Engine/Source/Developer/TraceInsights/Private/Insights/Widgets/STimingView.cpp
--- a/Engine/Source/Developer/TraceInsights/Private/Insights/Widgets/STimingView.cpp
+++ b/Engine/Source/Developer/TraceInsights/Private/Insights/Widgets/STimingView.cpp
@@ -35,7 +35,9 @@
 #include "Widgets/Text/STextBlock.h"
 
 // TraceServices
+#include "Common/ProviderLock.h"
 #include "TraceServices/Model/LoadTimeProfiler.h"
+#include "TraceServices/Model/Regions.h"
 
 // TraceInsightsCore
 #include "InsightsCore/Common/PaintUtils.h"
@@ -5250,6 +5252,19 @@ void STimingView::QuickFind_Execute()
 
 		NewFilterConfigurator->Add(MakeShared<FTimerNameFilter>());
 		NewFilterConfigurator->Add(MakeShared<FMetadataFilter>());
+		
+		TSharedRef<FFilterWithSuggestions> RegionNameFilter = MakeShared<FFilterWithSuggestions>(
+			static_cast<int32>(EFilterField::RegionName),
+			LOCTEXT("RegionName", "Region Name"),
+			LOCTEXT("RegionName", "Region Name"),
+			EFilterDataType::String,
+			nullptr,
+			FFilterService::Get()->GetStringOperators());
+		RegionNameFilter->SetCallback([this](const FString& Text, TArray<FString>& OutSuggestions)
+		{
+			this->PopulateRegionNameSuggestionList(Text, OutSuggestions);
+		});
+		NewFilterConfigurator->Add(RegionNameFilter);
 
 		for (Timing::ITimingViewExtender* Extender : GetExtenders())
 		{
@@ -5709,6 +5724,30 @@ void STimingView::PopulateTimerNameSuggestionList(const FString& Text, TArray<FS
 
 ////////////////////////////////////////////////////////////////////////////////////////////////////
 
+void STimingView::PopulateRegionNameSuggestionList(const FString& Text, TArray<FString>& OutSuggestions)
+{
+	TSharedPtr<const TraceServices::IAnalysisSession> Session = FInsightsManager::Get()->GetSession();
+	if (Session.IsValid())
+	{
+		const TraceServices::IRegionProvider& RegionProvider = TraceServices::ReadRegionProvider(*Session);
+		TraceServices::FProviderReadScopeLock RegionProviderScopedLock(RegionProvider);
+		
+		RegionProvider.EnumerateTimelinesByCategory([&Text, &OutSuggestions](const TraceServices::IRegionTimeline& Timeline, const TCHAR* Category)
+		{
+			Timeline.EnumerateUniqueRegionNames([&Text, &OutSuggestions](const TCHAR* RegionName) -> bool
+			{
+				if (FCString::Stristr(RegionName, *Text) || Text.IsEmpty())
+				{
+					OutSuggestions.Add(RegionName);
+				}
+				return true;
+			});
+		});
+	}
+}
+
+////////////////////////////////////////////////////////////////////////////////////////////////////
+
 void STimingView::EnumerateFilteredTracks(TSharedPtr<UE::Insights::FFilterConfigurator> InFilterConfigurator, TSharedPtr<const FBaseTimingTrack> PriorityTrack, EnumerateFilteredTracksCallback Callback)
 {
 	using EFilterField = UE::Insights::EFilterField;
diff --git a/Engine/Source/Developer/TraceInsights/Private/Insights/Widgets/STimingView.h b/Engine/Source/Developer/TraceInsights/Private/Insights/Widgets/STimingView.h
--- a/Engine/Source/Developer/TraceInsights/Private/Insights/Widgets/STimingView.h
+++ b/Engine/Source/Developer/TraceInsights/Private/Insights/Widgets/STimingView.h
@@ -568,6 +568,7 @@ PRAGMA_ENABLE_DEPRECATION_WARNINGS
 	TSharedRef<SDockTab> SpawnQuickFindTab(const FSpawnTabArgs& Args);
 	void PopulateTrackSuggestionList(const FString& Text, TArray<FString>& OutSuggestions);
 	void PopulateTimerNameSuggestionList(const FString& Text, TArray<FString>& OutSuggestions);
+	void PopulateRegionNameSuggestionList(const FString& Text, TArray<FString>& OutSuggestions);
 
 	typedef TFunctionRef<void(TSharedPtr<const FBaseTimingTrack> Track)> EnumerateFilteredTracksCallback;
 	void EnumerateFilteredTracks(TSharedPtr<UE::Insights::FFilterConfigurator> FilterConfigurator, TSharedPtr<const FBaseTimingTrack> PriorityTrack, EnumerateFilteredTracksCallback Callback);
diff --git a/Engine/Source/Developer/TraceServices/Private/Model/Regions.cpp b/Engine/Source/Developer/TraceServices/Private/Model/Regions.cpp
--- a/Engine/Source/Developer/TraceServices/Private/Model/Regions.cpp
+++ b/Engine/Source/Developer/TraceServices/Private/Model/Regions.cpp
@@ -41,6 +41,21 @@ const FRegionLane* FRegionTimeline::GetLane(int32 Index) const
 	return nullptr;
 }
 
+const bool FRegionTimeline::EnumerateUniqueRegionNames(TFunctionRef<bool(const TCHAR*)> Callback) const
+{
+	Provider.ReadAccessCheck();
+
+	for (const TCHAR* RegionName : TimelineUniqueRegions)
+	{
+		if (!Callback(RegionName))
+		{
+			return false;
+		}
+	}
+
+	return true;
+}
+
 bool FRegionTimeline::EnumerateRegions(double IntervalStart, double IntervalEnd,
 	TFunctionRef<bool(const FTimeRegion&)> Callback) const
 {
@@ -89,6 +104,8 @@ FTimeRegion* FRegionTimeline::InsertNewRegion(double BeginTime, const TCHAR* Nam
 	}
 	FRegionLane& NewLane = Lanes[Region.Depth];	
 	NewLane.Regions.EmplaceBack(Region);
+	TimelineUniqueRegions.FindOrAdd(Region.Text);
+
 	FTimeRegion* NewOpenRegion = &(Lanes[Region.Depth].Regions.Last());
 
 	return NewOpenRegion;
diff --git a/Engine/Source/Developer/TraceServices/Private/Model/RegionsPrivate.h b/Engine/Source/Developer/TraceServices/Private/Model/RegionsPrivate.h
--- a/Engine/Source/Developer/TraceServices/Private/Model/RegionsPrivate.h
+++ b/Engine/Source/Developer/TraceServices/Private/Model/RegionsPrivate.h
@@ -46,6 +46,14 @@ public:
 	 */
 	UE_API const FRegionLane* GetLane(int32 Index) const override;
 
+	/**
+	 * Enumerates all region names for a given timeline. Considers all lanes.
+	 * 
+	 * @param Callback a callback called for each region. Return false to abort iteration.
+	 * @returns true if the enumeration finished, false if it was aborted by the callback returning false
+	 */
+	UE_API const bool EnumerateUniqueRegionNames(TFunctionRef<bool(const TCHAR*)> Callback) const override;
+	
 	/**
 	 * Enumerates all regions that overlap a certain time interval. Will enumerate by depth but does not expose lanes.
 	 * 
@@ -72,6 +80,7 @@ private:
 	
 	// Closed regions
 	TArray<FRegionLane> Lanes;
+	TSet<const TCHAR*, TStringPointerSetKeyFuncs_DEPRECATED<const TCHAR*>> TimelineUniqueRegions;
 	const TCHAR* Category = nullptr;
 	TraceServices::IAnalysisSession& Session;
 	// the parent provider. Required to engage the ReadAccessLock or WriteAccessLock when a Lane is modified
diff --git a/Engine/Source/Developer/TraceServices/Public/TraceServices/Model/Regions.h b/Engine/Source/Developer/TraceServices/Public/TraceServices/Model/Regions.h
--- a/Engine/Source/Developer/TraceServices/Public/TraceServices/Model/Regions.h
+++ b/Engine/Source/Developer/TraceServices/Public/TraceServices/Model/Regions.h
@@ -68,6 +68,14 @@ public:
 	 */
 	virtual const FRegionLane* GetLane(int32 Index) const = 0;
 
+	/**
+	 * Enumerates all region names for a given timeline. Considers all lanes.
+	 * 
+	 * @param Callback a callback called for each region. Return false to abort iteration.
+	 * @returns true if the enumeration finished, false if it was aborted by the callback returning false
+	 */
+	virtual const bool EnumerateUniqueRegionNames(TFunctionRef<bool(const TCHAR*)> Callback) const = 0;
+
 	/**
 	 * Enumerates all regions that overlap a certain time interval. Will enumerate by depth but does not expose lanes.
 	 * 
