[Insights] Timing View / Quick Find: Added support for searching by region category in Quick Find

diff --git a/Engine/Source/Developer/TraceInsights/Private/Insights/TimingProfiler/Tracks/RegionsTimingTrack.cpp b/Engine/Source/Developer/TraceInsights/Private/Insights/TimingProfiler/Tracks/RegionsTimingTrack.cpp
--- a/Engine/Source/Developer/TraceInsights/Private/Insights/TimingProfiler/Tracks/RegionsTimingTrack.cpp
+++ b/Engine/Source/Developer/TraceInsights/Private/Insights/TimingProfiler/Tracks/RegionsTimingTrack.cpp
@@ -190,6 +190,7 @@ void FTimingRegionsTrack::BuildFilteredDrawState(ITimingEventsTrackDrawStateBuil
 		FilterContext.AddFilterData<double>(static_cast<int32>(EFilterField::Duration), 0.0f);
 		FilterContext.AddFilterData<FString>(static_cast<int32>(EFilterField::TrackName), this->GetName());
 		FilterContext.AddFilterData<FString>(static_cast<int32>(EFilterField::RegionName), "");
+		FilterContext.AddFilterData<FString>(static_cast<int32>(EFilterField::RegionCategory), "");
 
 		TSharedPtr<const TraceServices::IAnalysisSession> Session = FInsightsManager::Get()->GetSession();
 
@@ -211,6 +212,7 @@ void FTimingRegionsTrack::BuildFilteredDrawState(ITimingEventsTrackDrawStateBuil
 							FilterContext.SetFilterData<double>(static_cast<int32>(EFilterField::EndTime), Region.EndTime);
 							FilterContext.SetFilterData<double>(static_cast<int32>(EFilterField::Duration), Region.EndTime - Region.BeginTime);
 							FilterContext.SetFilterData<FString>(static_cast<int32>(EFilterField::RegionName), Region.Text);
+							FilterContext.SetFilterData<FString>(static_cast<int32>(EFilterField::RegionCategory), Region.Category);
 
 							RegionHadEvents = true;
 							if (FilterConfigurator->ApplyFilters(FilterContext))
@@ -264,6 +266,7 @@ bool FTimingRegionsTrack::FindRegionEvent(const FTimingEventSearchParameters& In
 	FilterConfiguratorContext.AddFilterData<double>(static_cast<int32>(EFilterField::Duration), 0.0f);
 	FilterConfiguratorContext.AddFilterData<FString>(static_cast<int32>(EFilterField::TrackName), this->GetName());
 	FilterConfiguratorContext.AddFilterData<FString>(static_cast<int32>(EFilterField::RegionName), "");
+	FilterConfiguratorContext.AddFilterData<FString>(static_cast<int32>(EFilterField::RegionCategory), "");
 
 	return TTimingEventSearch<TraceServices::FTimeRegion>::Search(
 	InParameters,
@@ -303,6 +306,7 @@ bool FTimingRegionsTrack::FindRegionEvent(const FTimingEventSearchParameters& In
 			FilterConfiguratorContext.SetFilterData<double>(static_cast<int32>(EFilterField::EndTime), EventEndTime);
 			FilterConfiguratorContext.SetFilterData<double>(static_cast<int32>(EFilterField::Duration), EventEndTime - EventStartTime);
 			FilterConfiguratorContext.SetFilterData<FString>(static_cast<int32>(EFilterField::RegionName), Region.Text);
+			FilterConfiguratorContext.SetFilterData<FString>(static_cast<int32>(EFilterField::RegionCategory), Region.Category);
 			return InParameters.FilterExecutor->ApplyFilters(FilterConfiguratorContext);
 		}
 
diff --git a/Engine/Source/Developer/TraceInsights/Private/Insights/Widgets/STimingView.cpp b/Engine/Source/Developer/TraceInsights/Private/Insights/Widgets/STimingView.cpp
--- a/Engine/Source/Developer/TraceInsights/Private/Insights/Widgets/STimingView.cpp
+++ b/Engine/Source/Developer/TraceInsights/Private/Insights/Widgets/STimingView.cpp
@@ -5266,6 +5266,19 @@ void STimingView::QuickFind_Execute()
 		});
 		NewFilterConfigurator->Add(RegionNameFilter);
 
+		TSharedRef<FFilterWithSuggestions> RegionCategoryFilter = MakeShared<FFilterWithSuggestions>(
+			static_cast<int32>(EFilterField::RegionCategory),
+			LOCTEXT("RegionCategory", "Region Category"),
+			LOCTEXT("RegionCategory", "Region Category"),
+			EFilterDataType::String,
+			nullptr,
+			FFilterService::Get()->GetStringOperators());
+		RegionCategoryFilter->SetCallback([this](const FString& Text, TArray<FString>& OutSuggestions)
+		{
+			this->PopulateRegionCategorySuggestionList(Text, OutSuggestions);
+		});
+		NewFilterConfigurator->Add(RegionCategoryFilter);
+
 		for (Timing::ITimingViewExtender* Extender : GetExtenders())
 		{
 			Extender->AddQuickFindFilters(NewFilterConfigurator);
@@ -5748,6 +5761,26 @@ void STimingView::PopulateRegionNameSuggestionList(const FString& Text, TArray<F
 
 ////////////////////////////////////////////////////////////////////////////////////////////////////
 
+void STimingView::PopulateRegionCategorySuggestionList(const FString& Text, TArray<FString>& OutSuggestions)
+{
+	TSharedPtr<const TraceServices::IAnalysisSession> Session = FInsightsManager::Get()->GetSession();
+	if (Session.IsValid())
+	{
+		const TraceServices::IRegionProvider& RegionProvider = TraceServices::ReadRegionProvider(*Session);
+		TraceServices::FProviderReadScopeLock RegionProviderScopedLock(RegionProvider);
+		
+		RegionProvider.EnumerateTimelinesByCategory([&Text, &OutSuggestions](const TraceServices::IRegionTimeline& Timeline, const TCHAR* Category)
+		{
+			if (FCString::Stristr(Category, *Text) || Text.IsEmpty())
+			{
+				OutSuggestions.Add(Category);
+			}
+		});
+	}
+}
+
+////////////////////////////////////////////////////////////////////////////////////////////////////
+
 void STimingView::EnumerateFilteredTracks(TSharedPtr<UE::Insights::FFilterConfigurator> InFilterConfigurator, TSharedPtr<const FBaseTimingTrack> PriorityTrack, EnumerateFilteredTracksCallback Callback)
 {
 	using EFilterField = UE::Insights::EFilterField;
diff --git a/Engine/Source/Developer/TraceInsights/Private/Insights/Widgets/STimingView.h b/Engine/Source/Developer/TraceInsights/Private/Insights/Widgets/STimingView.h
--- a/Engine/Source/Developer/TraceInsights/Private/Insights/Widgets/STimingView.h
+++ b/Engine/Source/Developer/TraceInsights/Private/Insights/Widgets/STimingView.h
@@ -569,6 +569,7 @@ PRAGMA_ENABLE_DEPRECATION_WARNINGS
 	void PopulateTrackSuggestionList(const FString& Text, TArray<FString>& OutSuggestions);
 	void PopulateTimerNameSuggestionList(const FString& Text, TArray<FString>& OutSuggestions);
 	void PopulateRegionNameSuggestionList(const FString& Text, TArray<FString>& OutSuggestions);
+	void PopulateRegionCategorySuggestionList(const FString& Text, TArray<FString>& OutSuggestions);
 
 	typedef TFunctionRef<void(TSharedPtr<const FBaseTimingTrack> Track)> EnumerateFilteredTracksCallback;
 	void EnumerateFilteredTracks(TSharedPtr<UE::Insights::FFilterConfigurator> FilterConfigurator, TSharedPtr<const FBaseTimingTrack> PriorityTrack, EnumerateFilteredTracksCallback Callback);
diff --git a/Engine/Source/Developer/TraceInsightsCore/Public/InsightsCore/Filter/ViewModels/Filters.h b/Engine/Source/Developer/TraceInsightsCore/Public/InsightsCore/Filter/ViewModels/Filters.h
--- a/Engine/Source/Developer/TraceInsightsCore/Public/InsightsCore/Filter/ViewModels/Filters.h
+++ b/Engine/Source/Developer/TraceInsightsCore/Public/InsightsCore/Filter/ViewModels/Filters.h
@@ -61,7 +61,8 @@ enum class EFilterField : int32
 	TimerName = 5,
 	CoreEventName = 6,
 	RegionName = 7,
-	Metadata = 8,
+	RegionCategory = 8,
+	Metadata = 9,
 };
 
 ////////////////////////////////////////////////////////////////////////////////////////////////////
diff --git a/Engine/Source/Developer/TraceServices/Private/Model/Regions.cpp b/Engine/Source/Developer/TraceServices/Private/Model/Regions.cpp
--- a/Engine/Source/Developer/TraceServices/Private/Model/Regions.cpp
+++ b/Engine/Source/Developer/TraceServices/Private/Model/Regions.cpp
@@ -355,19 +355,16 @@ TArray<FTimeRegion*> FRegionProvider::InsertNewRegion(double BeginTime, const TC
                                                       const TCHAR* Category)
 {
 	const TCHAR* StoredName = Session.StoreString(Name);
-	const TCHAR* StoredCategory = Category ? Session.StoreString(Category) : nullptr;
+	const TCHAR* StoredCategory = Category ? Session.StoreString(Category) : UncategorizedName;
 	
 	TArray<FTimeRegion*> Regions;
 	Regions.Add(AllRegions.InsertNewRegion(BeginTime, StoredName, Id, StoredCategory));
-
-	// skip per category timeline if no category is set
-	const TCHAR* Key = !StoredCategory ? UncategorizedName : StoredCategory;
 	
-	FRegionTimeline* Timeline = RegionsPerCategory.Find(Key);
+	FRegionTimeline* Timeline = RegionsPerCategory.Find(StoredCategory);
 	if (!Timeline)
 	{
-		FRegionTimeline NewTimeline{*this, Session, Key};
-		Timeline = &RegionsPerCategory.Emplace(Key, NewTimeline);
+		FRegionTimeline NewTimeline{*this, Session, StoredCategory};
+		Timeline = &RegionsPerCategory.Emplace(StoredCategory, NewTimeline);
 	}
 	Regions.Add(Timeline->InsertNewRegion(BeginTime, StoredName, Id, StoredCategory));
 	return Regions;
