[Insights] Added support for custom timelines for the ThreadTimingTrack.

#rnx

diff --git a/Engine/Source/Developer/TraceInsights/Private/Insights/TimingProfiler/Tracks/ThreadTimingTrack.cpp b/Engine/Source/Developer/TraceInsights/Private/Insights/TimingProfiler/Tracks/ThreadTimingTrack.cpp
--- a/Engine/Source/Developer/TraceInsights/Private/Insights/TimingProfiler/Tracks/ThreadTimingTrack.cpp
+++ b/Engine/Source/Developer/TraceInsights/Private/Insights/TimingProfiler/Tracks/ThreadTimingTrack.cpp
@@ -39,6 +39,12 @@ namespace UE::Insights::TimingProfiler
 
 INSIGHTS_IMPLEMENT_RTTI(FThreadTimingTrack);
 
+////////////////////////////////////////////////////////////////////////////////////////////////////
+// FThreadTimingTrackImpl
+////////////////////////////////////////////////////////////////////////////////////////////////////
+
+INSIGHTS_IMPLEMENT_RTTI(FThreadTimingTrackImpl)
+
 ////////////////////////////////////////////////////////////////////////////////////////////////////
 
 static void AppendMetadataToTooltip(FTooltipDrawState& Tooltip, const TraceServices::FMetadataSpec* MetadataSpec, TArrayView<const uint8>& Metadata)
@@ -340,7 +346,7 @@ static void AppendMetadataToString(FString& Str, const TraceServices::FMetadataS
 				Str += TEXT(", ");
 			}
 		}
-		
+
 		constexpr int MaxLength = 256;
 		TCHAR Data[MaxLength];
 
@@ -413,9 +419,15 @@ static void AppendMetadataToString(FString& Str, const TraceServices::FMetadataS
 
 ////////////////////////////////////////////////////////////////////////////////////////////////////
 
-static void AddTimingEventToBuilder(ITimingEventsTrackDrawStateBuilder& Builder,
-									double EventStartTime, double EventEndTime, uint32 EventDepth,
-									uint32 TimerIndex, const TraceServices::FTimingProfilerTimer* Timer)
+void AddTimingEventToBuilder(
+	ITimingEventsTrackDrawStateBuilder& Builder,
+	double EventStartTime,
+	double EventEndTime,
+	uint32 EventDepth,
+	uint32 TimerIndex,
+	const TraceServices::FTimingProfilerTimer* Timer,
+	const TraceServices::ITimingProfilerProvider& TimingProfilerProvider,
+	const TraceServices::ITimingProfilerTimerReader& TimerReader)
 {
 	if (EventDepth >= FTimingProfilerManager::Get()->GetEventDepthLimit())
 	{
@@ -450,7 +462,8 @@ static void AddTimingEventToBuilder(ITimingEventsTrackDrawStateBuilder& Builder,
 	}
 
 	Builder.AddEvent(EventStartTime, EventEndTime, EventDepth, EventColor,
-		[TimerIndex, Timer, EventStartTime, EventEndTime](float Width)
+		[TimerIndex, Timer, EventStartTime, EventEndTime, &TimingProfilerProvider, &TimerReader]
+		(float Width)
 		{
 			FString EventName = Timer->Name;
 
@@ -466,17 +479,7 @@ static void AddTimingEventToBuilder(ITimingEventsTrackDrawStateBuilder& Builder,
 				{
 					//EventName = TEXT("!") + EventName; // for debugging
 
-					TSharedPtr<const TraceServices::IAnalysisSession> Session = FInsightsManager::Get()->GetSession();
-					check(Session.IsValid());
-
-					//TraceServices::FAnalysisSessionReadScope SessionReadScope(*Session.Get());
-
-					const TraceServices::ITimingProfilerProvider& TimingProfilerProvider = *TraceServices::ReadTimingProfilerProvider(*Session.Get());
-
-					const TraceServices::ITimingProfilerTimerReader* TimerReader;
-					TimingProfilerProvider.ReadTimers([&TimerReader](const TraceServices::ITimingProfilerTimerReader& Out) { TimerReader = &Out; });
-
-					TArrayView<const uint8> Metadata = TimerReader->GetMetadata(TimerIndex);
+					TArrayView<const uint8> Metadata = TimerReader.GetMetadata(TimerIndex);
 
 					if (Metadata.Num() > 0)
 					{
@@ -499,71 +502,67 @@ static void AddTimingEventToBuilder(ITimingEventsTrackDrawStateBuilder& Builder,
 			}
 
 			return EventName;
-		});
+		}
+	); // Builder.AddEvent
 }
 
-////////////////////////////////////////////////////////////////////////////////////////////////////
-// FThreadTimingTrackImpl
-////////////////////////////////////////////////////////////////////////////////////////////////////
-
-INSIGHTS_IMPLEMENT_RTTI(FThreadTimingTrackImpl)
-
 ////////////////////////////////////////////////////////////////////////////////////////////////////
 
 void FThreadTimingTrackImpl::BuildDrawState(ITimingEventsTrackDrawStateBuilder& Builder, const ITimingTrackUpdateContext& Context)
 {
-	TSharedPtr<const TraceServices::IAnalysisSession> Session = FInsightsManager::Get()->GetSession();
-	if (Session.IsValid() && TraceServices::ReadTimingProfilerProvider(*Session.Get()))
-	{
-		TraceServices::FAnalysisSessionReadScope SessionReadScope(*Session.Get());
-
-		const TraceServices::ITimingProfilerProvider& TimingProfilerProvider = *TraceServices::ReadTimingProfilerProvider(*Session.Get());
+	const FTimingTrackViewport& Viewport = Context.GetViewport();
 
-		const TraceServices::ITimingProfilerTimerReader* TimerReader;
-		TimingProfilerProvider.ReadTimers([&TimerReader](const TraceServices::ITimingProfilerTimerReader& Out) { TimerReader = &Out; });
-
-		const FTimingTrackViewport& Viewport = Context.GetViewport();
-
-		TimingProfilerProvider.ReadTimeline(TimelineIndex,
-			[&Viewport, this, &Builder, TimerReader](const TraceServices::ITimingProfilerProvider::Timeline& Timeline)
-			{
-				if (FTimingEventsTrack::bUseDownSampling)
-				{
-					const double SecondsPerPixel = 1.0 / Viewport.GetScaleX();
-					Timeline.EnumerateEventsDownSampled(Viewport.GetStartTime(), Viewport.GetEndTime(), SecondsPerPixel,
-						[this, &Builder, TimerReader](double StartTime, double EndTime, uint32 Depth, const TraceServices::FTimingProfilerEvent& Event)
-						{
-							const TraceServices::FTimingProfilerTimer* Timer = TimerReader->GetTimer(Event.TimerIndex);
-							if (ensure(Timer != nullptr))
-							{
-								AddTimingEventToBuilder(Builder, StartTime, EndTime, Depth, Event.TimerIndex, Timer);
-							}
-							else
-							{
-								Builder.AddEvent(StartTime, EndTime, Depth, 0xFF000000, [&Event](float) { return FString::Printf(TEXT("[%u]"), Event.TimerIndex); });
-							}
-							return TraceServices::EEventEnumerate::Continue;
-						});
-				}
-				else
+	ReadTimers(
+		[this, &Builder, &Viewport]
+		(const TraceServices::ITimingProfilerProvider& TimingProfilerProvider, const TraceServices::ITimingProfilerTimerReader& TimerReader)
+		{
+			ReadTimeline(
+				[this, &Builder, &Viewport, &TimingProfilerProvider, &TimerReader]
+				(const TraceServices::ITimingProfilerProvider::Timeline& Timeline)
 				{
-					Timeline.EnumerateEvents(Viewport.GetStartTime(), Viewport.GetEndTime(),
-						[this, &Builder, TimerReader](double StartTime, double EndTime, uint32 Depth, const TraceServices::FTimingProfilerEvent& Event)
-						{
-							const TraceServices::FTimingProfilerTimer* Timer = TimerReader->GetTimer(Event.TimerIndex);
-							if (ensure(Timer != nullptr))
+					if (FTimingEventsTrack::bUseDownSampling)
+					{
+						const double SecondsPerPixel = 1.0 / Viewport.GetScaleX();
+						Timeline.EnumerateEventsDownSampled(Viewport.GetStartTime(), Viewport.GetEndTime(), SecondsPerPixel,
+							[this, &Builder, &TimingProfilerProvider, &TimerReader]
+							(double StartTime, double EndTime, uint32 Depth, const TraceServices::FTimingProfilerEvent& Event)
 							{
-								AddTimingEventToBuilder(Builder, StartTime, EndTime, Depth, Event.TimerIndex, Timer);
+								const TraceServices::FTimingProfilerTimer* Timer = TimerReader.GetTimer(Event.TimerIndex);
+								if (ensure(Timer != nullptr))
+								{
+									AddTimingEventToBuilder(Builder, StartTime, EndTime, Depth, Event.TimerIndex, Timer, TimingProfilerProvider, TimerReader);
+								}
+								else
+								{
+									Builder.AddEvent(StartTime, EndTime, Depth, 0xFF000000, [&Event](float) { return FString::Printf(TEXT("[%u]"), Event.TimerIndex); });
+								}
+								return TraceServices::EEventEnumerate::Continue;
 							}
-							else
+						);
+					}
+					else
+					{
+						Timeline.EnumerateEvents(Viewport.GetStartTime(), Viewport.GetEndTime(),
+							[this, &Builder, &TimingProfilerProvider, &TimerReader]
+							(double StartTime, double EndTime, uint32 Depth, const TraceServices::FTimingProfilerEvent& Event)
 							{
-								Builder.AddEvent(StartTime, EndTime, Depth, 0xFF000000, [&Event](float) { return FString::Printf(TEXT("[%u]"), Event.TimerIndex); });
+								const TraceServices::FTimingProfilerTimer* Timer = TimerReader.GetTimer(Event.TimerIndex);
+								if (ensure(Timer != nullptr))
+								{
+									AddTimingEventToBuilder(Builder, StartTime, EndTime, Depth, Event.TimerIndex, Timer, TimingProfilerProvider, TimerReader);
+								}
+								else
+								{
+									Builder.AddEvent(StartTime, EndTime, Depth, 0xFF000000, [&Event](float) { return FString::Printf(TEXT("[%u]"), Event.TimerIndex); });
+								}
+								return TraceServices::EEventEnumerate::Continue;
 							}
-							return TraceServices::EEventEnumerate::Continue;
-						});
+						);
+					}
 				}
-			});
-	}
+			); // ReadTimeline
+		}
+	); // ReadTimers
 }
 
 ////////////////////////////////////////////////////////////////////////////////////////////////////
@@ -590,208 +589,216 @@ void FThreadTimingTrackImpl::BuildFilteredDrawState(ITimingEventsTrackDrawStateB
 			FilterEventType = EventFilter.GetEventType();
 		}
 
-		TSharedPtr<const TraceServices::IAnalysisSession> Session = FInsightsManager::Get()->GetSession();
-		if (Session.IsValid() && TraceServices::ReadTimingProfilerProvider(*Session.Get()))
-		{
-			TraceServices::FAnalysisSessionReadScope SessionReadScope(*Session.Get());
-
-			const TraceServices::ITimingProfilerProvider& TimingProfilerProvider = *TraceServices::ReadTimingProfilerProvider(*Session.Get());
-
-			const TraceServices::ITimingProfilerTimerReader* TimerReader;
-			TimingProfilerProvider.ReadTimers([&TimerReader](const TraceServices::ITimingProfilerTimerReader& Out) { TimerReader = &Out; });
-
-			const FTimingTrackViewport& Viewport = Context.GetViewport();
+		const FTimingTrackViewport& Viewport = Context.GetViewport();
 
-			if (bFilterOnlyByEventType)
+		ReadTimers(
+			[this, &Builder, &EventFilterPtr, bFilterOnlyByEventType, FilterEventType, &Viewport]
+			(const TraceServices::ITimingProfilerProvider& TimingProfilerProvider, const TraceServices::ITimingProfilerTimerReader& TimerReader)
 			{
-				//TODO: Add a setting to switch this on/off
-				if (true)
+				if (bFilterOnlyByEventType)
 				{
-					TimingProfilerProvider.ReadTimeline(TimelineIndex,
-						[&Viewport, this, &Builder, TimerReader, FilterEventType](const TraceServices::ITimingProfilerProvider::Timeline& Timeline)
-						{
-							TArray<TArray<FPendingEventInfo>> FilteredEvents;
-
-							TraceServices::ITimeline<TraceServices::FTimingProfilerEvent>::EnumerateAsyncParams Params;
-							Params.IntervalStart = Viewport.GetStartTime();
-							Params.IntervalEnd = Viewport.GetEndTime();
-							Params.Resolution = 0.0;
-							Params.SetupCallback = [&FilteredEvents](uint32 NumTasks)
-							{
-								FilteredEvents.AddDefaulted(NumTasks);
-							};
-							Params.EventRangeCallback = [this, &Builder, TimerReader, FilterEventType, &FilteredEvents](double StartTime, double EndTime, uint32 Depth, const TraceServices::FTimingProfilerEvent& Event, uint32 TaskIndex)
+					//TODO: Add a setting to switch this on/off
+					if (true)
+					{
+						ReadTimeline(
+							[this, &Builder, FilterEventType, &Viewport, &TimingProfilerProvider, &TimerReader]
+							(const TraceServices::ITimingProfilerProvider::Timeline& Timeline)
 							{
-								const TraceServices::FTimingProfilerTimer* Timer = TimerReader->GetTimer(Event.TimerIndex);
-								if (ensure(Timer != nullptr))
-								{
-									if (Timer->Id == FilterEventType)
+								TArray<TArray<FPendingEventInfo>> FilteredEvents;
+
+								TraceServices::ITimeline<TraceServices::FTimingProfilerEvent>::EnumerateAsyncParams Params;
+								Params.IntervalStart = Viewport.GetStartTime();
+								Params.IntervalEnd = Viewport.GetEndTime();
+								Params.Resolution = 0.0;
+								Params.SetupCallback =
+									[&FilteredEvents](uint32 NumTasks)
 									{
-										FPendingEventInfo TimelineEvent;
-										TimelineEvent.StartTime = StartTime;
-										TimelineEvent.EndTime = EndTime;
-										TimelineEvent.Depth = Depth;
-										TimelineEvent.TimerIndex = Event.TimerIndex;
-										FilteredEvents[TaskIndex].Add(TimelineEvent);
-									}
-								}
-								return TraceServices::EEventEnumerate::Continue;
-							};
+										FilteredEvents.AddDefaulted(NumTasks);
+									};
+								Params.EventRangeCallback =
+									[this, &Builder, FilterEventType, &TimerReader, &FilteredEvents]
+									(double StartTime, double EndTime, uint32 Depth, const TraceServices::FTimingProfilerEvent& Event, uint32 TaskIndex)
+									{
+										const TraceServices::FTimingProfilerTimer* Timer = TimerReader.GetTimer(Event.TimerIndex);
+										if (ensure(Timer != nullptr))
+										{
+											if (Timer->Id == FilterEventType)
+											{
+												FPendingEventInfo TimelineEvent;
+												TimelineEvent.StartTime = StartTime;
+												TimelineEvent.EndTime = EndTime;
+												TimelineEvent.Depth = Depth;
+												TimelineEvent.TimerIndex = Event.TimerIndex;
+												FilteredEvents[TaskIndex].Add(TimelineEvent);
+											}
+										}
+										return TraceServices::EEventEnumerate::Continue;
+									};
 
-							// Note: Enumerating events for filtering should not use downsampling.
-							Timeline.EnumerateEventsDownSampledAsync(Params);
+								// Note: Enumerating events for filtering should not use downsampling.
+								Timeline.EnumerateEventsDownSampledAsync(Params);
 
-							for (TArray<FPendingEventInfo>& Array : FilteredEvents)
-							{
-								for (FPendingEventInfo& TimelineEvent : Array)
+								for (TArray<FPendingEventInfo>& Array : FilteredEvents)
 								{
-									const TraceServices::FTimingProfilerTimer* Timer = TimerReader->GetTimer(TimelineEvent.TimerIndex);
-									AddTimingEventToBuilder(Builder, TimelineEvent.StartTime, TimelineEvent.EndTime, TimelineEvent.Depth, TimelineEvent.TimerIndex, Timer);
+									for (FPendingEventInfo& TimelineEvent : Array)
+									{
+										const TraceServices::FTimingProfilerTimer* Timer = TimerReader.GetTimer(TimelineEvent.TimerIndex);
+										AddTimingEventToBuilder(Builder, TimelineEvent.StartTime, TimelineEvent.EndTime, TimelineEvent.Depth, TimelineEvent.TimerIndex, Timer, TimingProfilerProvider, TimerReader);
+									}
 								}
+							}); // ReadTimeline
+					}
+					else
+					{
+						ReadTimeline(
+							[this, &Builder, FilterEventType, &Viewport, &TimingProfilerProvider, &TimerReader]
+							(const TraceServices::ITimingProfilerProvider::Timeline& Timeline)
+							{
+								// Note: Enumerating events for filtering should not use downsampling.
+								Timeline.EnumerateEventsDownSampled(Viewport.GetStartTime(), Viewport.GetEndTime(), 0,
+									[this, &Builder, FilterEventType, &TimingProfilerProvider, &TimerReader]
+									(double StartTime, double EndTime, uint32 Depth, const TraceServices::FTimingProfilerEvent& Event)
+									{
+										const TraceServices::FTimingProfilerTimer* Timer = TimerReader.GetTimer(Event.TimerIndex);
+										if (ensure(Timer != nullptr))
+										{
+											if (Timer->Id == FilterEventType)
+											{
+												AddTimingEventToBuilder(Builder, StartTime, EndTime, Depth, Event.TimerIndex, Timer, TimingProfilerProvider, TimerReader);
+											}
+										}
+										return TraceServices::EEventEnumerate::Continue;
+									}
+								);
 							}
-						});
+						); // ReadTimeline
+					}
 				}
-				else
+				else // generic filter
 				{
-					TimingProfilerProvider.ReadTimeline(TimelineIndex,
-						[&Viewport, this, &Builder, TimerReader, FilterEventType](const TraceServices::ITimingProfilerProvider::Timeline& Timeline)
+					ReadTimeline(
+						[this, &Builder, &EventFilterPtr, &Viewport, &TimingProfilerProvider, &TimerReader]
+						(const TraceServices::ITimingProfilerProvider::Timeline& Timeline)
 						{
 							// Note: Enumerating events for filtering should not use downsampling.
-							Timeline.EnumerateEventsDownSampled(Viewport.GetStartTime(), Viewport.GetEndTime(), 0,
-								[this, &Builder, TimerReader, FilterEventType](double StartTime, double EndTime, uint32 Depth, const TraceServices::FTimingProfilerEvent& Event)
+							//const double SecondsPerPixel = 1.0 / Viewport.GetScaleX();
+							//Timeline.EnumerateEventsDownSampled(Viewport.GetStartTime(), Viewport.GetEndTime(), SecondsPerPixel,
+							Timeline.EnumerateEvents(Viewport.GetStartTime(), Viewport.GetEndTime(),
+								[this, &Builder, &EventFilterPtr, &TimingProfilerProvider, &TimerReader]
+								(double StartTime, double EndTime, uint32 Depth, const TraceServices::FTimingProfilerEvent& Event)
 								{
-									const TraceServices::FTimingProfilerTimer* Timer = TimerReader->GetTimer(Event.TimerIndex);
+									const TraceServices::FTimingProfilerTimer* Timer = TimerReader.GetTimer(Event.TimerIndex);
 									if (ensure(Timer != nullptr))
 									{
-										if (Timer->Id == FilterEventType)
+										FThreadTrackEvent TimingEvent(SharedThis(this), StartTime, EndTime, Depth);
+										TimingEvent.SetTimerId(Timer->Id);
+										TimingEvent.SetTimerIndex(Event.TimerIndex);
+
+										if (EventFilterPtr->FilterEvent(TimingEvent))
 										{
-											AddTimingEventToBuilder(Builder, StartTime, EndTime, Depth, Event.TimerIndex, Timer);
+											AddTimingEventToBuilder(Builder, StartTime, EndTime, Depth, Event.TimerIndex, Timer, TimingProfilerProvider, TimerReader);
 										}
 									}
 									return TraceServices::EEventEnumerate::Continue;
-								});
-						});
+								}
+							);
+						}
+					); // ReadTimeline
 				}
 			}
-			else // generic filter
+		); // ReadTimers
+	}
+
+	if (HasCustomFilter()) // Custom filter (from the filtering widget)
+	{
+		const FTimingTrackViewport& Viewport = Context.GetViewport();
+
+		ReadTimers(
+			[this, &Builder, &Viewport]
+			(const TraceServices::ITimingProfilerProvider& TimingProfilerProvider, const TraceServices::ITimingProfilerTimerReader& TimerReader)
 			{
-				TimingProfilerProvider.ReadTimeline(TimelineIndex,
-					[&Viewport, this, &Builder, TimerReader, &EventFilterPtr](const TraceServices::ITimingProfilerProvider::Timeline& Timeline)
+				ReadTimeline(
+					[this, &Builder, &Viewport, &TimingProfilerProvider, &TimerReader]
+					(const TraceServices::ITimingProfilerProvider::Timeline& Timeline)
 					{
+						TArray<TArray<FPendingEventInfo>> FilteredEvents;
+						TArray<FFilterContext> FilterContexts;
+
+						TraceServices::ITimeline<TraceServices::FTimingProfilerEvent>::EnumerateAsyncParams Params;
+						constexpr uint32 LargeTimelineThreshold = 50 * 1000 * 1000;
+						if (Timeline.GetEventCount() > LargeTimelineThreshold)
+						{
+							if (FilterConfigurator->IsKeyUsed(static_cast<int32>(EFilterField::Metadata)))
+							{
+								Params.MaxOccupancy = 0.75f; // This filter can be slow so reduce occupancy to avoid starvation.
+							}
+						}
+						Params.IntervalStart = Viewport.GetStartTime();
+						Params.IntervalEnd = Viewport.GetEndTime();
+
 						// Note: Enumerating events for filtering should not use downsampling.
-						//const double SecondsPerPixel = 1.0 / Viewport.GetScaleX();
-						//Timeline.EnumerateEventsDownSampled(Viewport.GetStartTime(), Viewport.GetEndTime(), SecondsPerPixel,
-						Timeline.EnumerateEvents(Viewport.GetStartTime(), Viewport.GetEndTime(),
-							[this, &Builder, TimerReader, &EventFilterPtr](double StartTime, double EndTime, uint32 Depth, const TraceServices::FTimingProfilerEvent& Event)
+						Params.Resolution = 0.0;
+						Params.SetupCallback =
+							[this, &FilteredEvents, &FilterContexts]
+							(uint32 NumTasks)
 							{
-								const TraceServices::FTimingProfilerTimer* Timer = TimerReader->GetTimer(Event.TimerIndex);
+								FilteredEvents.AddDefaulted(NumTasks);
+								FilterContexts.AddDefaulted(NumTasks);
+								for (FFilterContext& Context : FilterContexts)
+								{
+									Context.SetReturnValueForUnsetFilters(false);
+									Context.AddFilterData<double>(static_cast<int32>(EFilterField::StartTime), 0.0f);
+									Context.AddFilterData<double>(static_cast<int32>(EFilterField::EndTime), 0.0f);
+									Context.AddFilterData<double>(static_cast<int32>(EFilterField::Duration), 0.0f);
+									Context.AddFilterData<FString>(static_cast<int32>(EFilterField::TrackName), this->GetName());
+									Context.AddFilterData<int64>(static_cast<int32>(EFilterField::TimerId), 0);
+									Context.AddFilterData<int64>(static_cast<int32>(EFilterField::TimerName), 0);
+									Context.AddFilterData<int64>(static_cast<int32>(EFilterField::Metadata), 0);
+								}
+							};
+						Params.EventRangeCallback =
+							[this, &Builder, &TimerReader, &FilteredEvents, &FilterContexts]
+							(double StartTime, double EndTime, uint32 Depth, const TraceServices::FTimingProfilerEvent& Event, uint32 TaskIndex)
+							{
+								const TraceServices::FTimingProfilerTimer* Timer = TimerReader.GetTimer(Event.TimerIndex);
 								if (ensure(Timer != nullptr))
 								{
-									FThreadTrackEvent TimingEvent(SharedThis(this), StartTime, EndTime, Depth);
-									TimingEvent.SetTimerId(Timer->Id);
-									TimingEvent.SetTimerIndex(Event.TimerIndex);
-
-									if (EventFilterPtr->FilterEvent(TimingEvent))
+									FFilterContext& Context = FilterContexts[TaskIndex];
+									Context.SetFilterData<double>(static_cast<int32>(EFilterField::StartTime), StartTime);
+									Context.SetFilterData<double>(static_cast<int32>(EFilterField::EndTime), EndTime);
+									Context.SetFilterData<double>(static_cast<int32>(EFilterField::Duration), EndTime - StartTime);
+									// The TimerName filter also translates to the numeric Id for performance reasons.
+									Context.SetFilterData<int64>(static_cast<int32>(EFilterField::TimerId), Timer->Id);
+									Context.SetFilterData<int64>(static_cast<int32>(EFilterField::TimerName), Timer->Id);
+									Context.SetFilterData<int64>(static_cast<int32>(EFilterField::Metadata), Event.TimerIndex);
+
+									if (FilterConfigurator->ApplyFilters(Context))
 									{
-										AddTimingEventToBuilder(Builder, StartTime, EndTime, Depth, Event.TimerIndex, Timer);
+										FPendingEventInfo TimelineEvent;
+										TimelineEvent.StartTime = StartTime;
+										TimelineEvent.EndTime = EndTime;
+										TimelineEvent.Depth = Depth;
+										TimelineEvent.TimerIndex = Event.TimerIndex;
+										FilteredEvents[TaskIndex].Add(TimelineEvent);
 									}
 								}
 								return TraceServices::EEventEnumerate::Continue;
-							});
-					});
-			}
-		}
-	}
-
-	if (HasCustomFilter()) // Custom filter (from the filtering widget)
-	{
-		TSharedPtr<const TraceServices::IAnalysisSession> Session = FInsightsManager::Get()->GetSession();
-		if (Session.IsValid() && TraceServices::ReadTimingProfilerProvider(*Session.Get()))
-		{
-			TraceServices::FAnalysisSessionReadScope SessionReadScope(*Session.Get());
-
-			const TraceServices::ITimingProfilerProvider& TimingProfilerProvider = *TraceServices::ReadTimingProfilerProvider(*Session.Get());
-
-			const TraceServices::ITimingProfilerTimerReader* TimerReader;
-			TimingProfilerProvider.ReadTimers([&TimerReader](const TraceServices::ITimingProfilerTimerReader& Out) { TimerReader = &Out; });
-
-			const FTimingTrackViewport& Viewport = Context.GetViewport();
-
-			TimingProfilerProvider.ReadTimeline(TimelineIndex,
-				[&Viewport, this, &Builder, TimerReader](const TraceServices::ITimingProfilerProvider::Timeline& Timeline)
-				{
-					TArray<TArray<FPendingEventInfo>> FilteredEvents;
-					TArray<FFilterContext> FilterContexts;
+							};
 
-					TraceServices::ITimeline<TraceServices::FTimingProfilerEvent>::EnumerateAsyncParams Params;
-					constexpr uint32 LargeTimelineThreshold = 50 * 1000 * 1000;
-					if (Timeline.GetEventCount() > LargeTimelineThreshold)
-					{
-						if (FilterConfigurator->IsKeyUsed(static_cast<int32>(EFilterField::Metadata)))
-						{
-							Params.MaxOccupancy = 0.75f; // This filter can be slow so reduce occupancy to avoid starvation.
-						}
-					}
-					Params.IntervalStart = Viewport.GetStartTime();
-					Params.IntervalEnd = Viewport.GetEndTime();
+						Timeline.EnumerateEventsDownSampledAsync(Params);
 
-					// Note: Enumerating events for filtering should not use downsampling.
-					Params.Resolution = 0.0;
-					Params.SetupCallback = [&FilteredEvents, &FilterContexts, this](uint32 NumTasks)
-					{
-						FilteredEvents.AddDefaulted(NumTasks);
-						FilterContexts.AddDefaulted(NumTasks);
-						for (FFilterContext& Context : FilterContexts)
+						for (TArray<FPendingEventInfo>& Array : FilteredEvents)
 						{
-							Context.SetReturnValueForUnsetFilters(false);
-							Context.AddFilterData<double>(static_cast<int32>(EFilterField::StartTime), 0.0f);
-							Context.AddFilterData<double>(static_cast<int32>(EFilterField::EndTime), 0.0f);
-							Context.AddFilterData<double>(static_cast<int32>(EFilterField::Duration), 0.0f);
-							Context.AddFilterData<FString>(static_cast<int32>(EFilterField::TrackName), this->GetName());
-							Context.AddFilterData<int64>(static_cast<int32>(EFilterField::TimerId), 0);
-							Context.AddFilterData<int64>(static_cast<int32>(EFilterField::TimerName), 0);
-							Context.AddFilterData<int64>(static_cast<int32>(EFilterField::Metadata), 0);
-						}
-					};
-					Params.EventRangeCallback = [this, &Builder, TimerReader, &FilteredEvents, &FilterContexts](double StartTime, double EndTime, uint32 Depth, const TraceServices::FTimingProfilerEvent& Event, uint32 TaskIndex)
-					{
-						const TraceServices::FTimingProfilerTimer* Timer = TimerReader->GetTimer(Event.TimerIndex);
-						if (ensure(Timer != nullptr))
-						{
-							FFilterContext& Context = FilterContexts[TaskIndex];
-							Context.SetFilterData<double>(static_cast<int32>(EFilterField::StartTime), StartTime);
-							Context.SetFilterData<double>(static_cast<int32>(EFilterField::EndTime), EndTime);
-							Context.SetFilterData<double>(static_cast<int32>(EFilterField::Duration), EndTime - StartTime);
-							// The TimerName filter also translates to the numeric Id for performance reasons.
-							Context.SetFilterData<int64>(static_cast<int32>(EFilterField::TimerId), Timer->Id);
-							Context.SetFilterData<int64>(static_cast<int32>(EFilterField::TimerName), Timer->Id);
-							Context.SetFilterData<int64>(static_cast<int32>(EFilterField::Metadata), Event.TimerIndex);
-
-							if (FilterConfigurator->ApplyFilters(Context))
+							for (FPendingEventInfo& TimelineEvent : Array)
 							{
-								FPendingEventInfo TimelineEvent;
-								TimelineEvent.StartTime = StartTime;
-								TimelineEvent.EndTime = EndTime;
-								TimelineEvent.Depth = Depth;
-								TimelineEvent.TimerIndex = Event.TimerIndex;
-								FilteredEvents[TaskIndex].Add(TimelineEvent);
+								const TraceServices::FTimingProfilerTimer* Timer = TimerReader.GetTimer(TimelineEvent.TimerIndex);
+								AddTimingEventToBuilder(Builder, TimelineEvent.StartTime, TimelineEvent.EndTime, TimelineEvent.Depth, TimelineEvent.TimerIndex, Timer, TimingProfilerProvider, TimerReader);
 							}
 						}
-						return TraceServices::EEventEnumerate::Continue;
-					};
-
-					Timeline.EnumerateEventsDownSampledAsync(Params);
-
-					for (TArray<FPendingEventInfo>& Array : FilteredEvents)
-					{
-						for (FPendingEventInfo& TimelineEvent : Array)
-						{
-							const TraceServices::FTimingProfilerTimer* Timer = TimerReader->GetTimer(TimelineEvent.TimerIndex);
-							AddTimingEventToBuilder(Builder, TimelineEvent.StartTime, TimelineEvent.EndTime, TimelineEvent.Depth, TimelineEvent.TimerIndex, Timer);
-						}
 					}
-				});
-		}
+				); // ReadTimeline
+			}
+		); // ReadTimers
 	}
 }
 
@@ -812,41 +819,37 @@ void FThreadTimingTrackImpl::PostDraw(const ITimingTrackDrawContext& Context) co
 		const FThreadTrackEvent& SelectedEvent = SelectedEventPtr->As<FThreadTrackEvent>();
 		const ITimingViewDrawHelper& Helper = Context.GetHelper();
 
-		TSharedPtr<const TraceServices::IAnalysisSession> Session = FInsightsManager::Get()->GetSession();
-		check(Session.IsValid());
-
-		TraceServices::FAnalysisSessionReadScope SessionReadScope(*Session.Get());
-
-		const TraceServices::ITimingProfilerProvider& TimingProfilerProvider = *TraceServices::ReadTimingProfilerProvider(*Session.Get());
-
-		const TraceServices::ITimingProfilerTimerReader* TimerReader;
-		TimingProfilerProvider.ReadTimers([&TimerReader](const TraceServices::ITimingProfilerTimerReader& Out) { TimerReader = &Out; });
-
-		const TraceServices::FTimingProfilerTimer* Timer = TimerReader->GetTimer(SelectedEvent.GetTimerIndex());
-		if (Timer != nullptr)
-		{
-			FString TimerName(Timer->Name);
-
-			const double SelectedEventDuration = SelectedEvent.GetDuration();
-			TStringBuilder<1024> StringBuilder;
-			StringBuilder.Appendf(TEXT("  Incl.: %s"), *FormatTimeAuto(SelectedEventDuration, 2));
-			if (SelectedEventDuration != std::numeric_limits<double>::infinity())
+		ReadTimers(
+			[this, &Context, &SelectedEvent, &Helper]
+			(const TraceServices::ITimingProfilerProvider& TimingProfilerProvider, const TraceServices::ITimingProfilerTimerReader& TimerReader)
 			{
-				StringBuilder.Appendf(TEXT("   Excl.: %s"), *FormatTimeAuto(SelectedEvent.GetExclusiveTime(), 2));
-			}
-			FString StatsText(StringBuilder.ToView());
+				const TraceServices::FTimingProfilerTimer* Timer = TimerReader.GetTimer(SelectedEvent.GetTimerIndex());
+				if (Timer != nullptr)
+				{
+					FString TimerName(Timer->Name);
 
-			if (Timer->File)
-			{
-				FString SourceFile = FPaths::GetCleanFilename(FString(Timer->File));
-				FString SourceFileAndLine = FString::Printf(TEXT("%s (%d)"), *SourceFile, Timer->Line);
-				DrawSelectedEventInfoEx(StatsText, TimerName, SourceFileAndLine, Context.GetViewport(), Context.GetDrawContext(), Helper.GetWhiteBrush(), Helper.GetEventFont());
-			}
-			else
-			{
-				DrawSelectedEventInfoEx(StatsText, TimerName, FString(), Context.GetViewport(), Context.GetDrawContext(), Helper.GetWhiteBrush(), Helper.GetEventFont());
+					const double SelectedEventDuration = SelectedEvent.GetDuration();
+					TStringBuilder<1024> StringBuilder;
+					StringBuilder.Appendf(TEXT("  Incl.: %s"), *FormatTimeAuto(SelectedEventDuration, 2));
+					if (SelectedEventDuration != std::numeric_limits<double>::infinity())
+					{
+						StringBuilder.Appendf(TEXT("   Excl.: %s"), *FormatTimeAuto(SelectedEvent.GetExclusiveTime(), 2));
+					}
+					FString StatsText(StringBuilder.ToView());
+
+					if (Timer->File)
+					{
+						FString SourceFile = FPaths::GetCleanFilename(FString(Timer->File));
+						FString SourceFileAndLine = FString::Printf(TEXT("%s (%d)"), *SourceFile, Timer->Line);
+						DrawSelectedEventInfoEx(StatsText, TimerName, SourceFileAndLine, Context.GetViewport(), Context.GetDrawContext(), Helper.GetWhiteBrush(), Helper.GetEventFont());
+					}
+					else
+					{
+						DrawSelectedEventInfoEx(StatsText, TimerName, FString(), Context.GetViewport(), Context.GetDrawContext(), Helper.GetWhiteBrush(), Helper.GetEventFont());
+					}
+				}
 			}
-		}
+		); // ReadTimers
 	}
 }
 
@@ -867,74 +870,74 @@ void FThreadTimingTrackImpl::InitTooltip(FTooltipDrawState& InOutTooltip, const
 		TSharedPtr<FThreadTrackEvent> RootTimingEvent;
 		GetParentAndRoot(TooltipEvent, ParentTimingEvent, RootTimingEvent);
 
-		TSharedPtr<const TraceServices::IAnalysisSession> Session = FInsightsManager::Get()->GetSession();
-		check(Session.IsValid());
-
-		TraceServices::FAnalysisSessionReadScope SessionReadScope(*Session.Get());
-
-		const TraceServices::ITimingProfilerProvider& TimingProfilerProvider = *TraceServices::ReadTimingProfilerProvider(*Session.Get());
-
-		const TraceServices::ITimingProfilerTimerReader* TimerReader;
-		TimingProfilerProvider.ReadTimers([&TimerReader](const TraceServices::ITimingProfilerTimerReader& Out) { TimerReader = &Out; });
-
-		const TraceServices::FTimingProfilerTimer* Timer = TimerReader->GetTimer(TooltipEvent.GetTimerIndex());
-		const TCHAR* TimerName = (Timer != nullptr) ? Timer->Name : TEXT("N/A");
-		InOutTooltip.AddTitle(TimerName);
+		ReadTimers(
+			[this, &InOutTooltip, &TooltipEvent, &ParentTimingEvent, &RootTimingEvent]
+			(const TraceServices::ITimingProfilerProvider& TimingProfilerProvider, const TraceServices::ITimingProfilerTimerReader& TimerReader)
+			{
+				const TraceServices::FTimingProfilerTimer* Timer = TimerReader.GetTimer(TooltipEvent.GetTimerIndex());
+				const TCHAR* TimerName = (Timer != nullptr) ? Timer->Name : TEXT("N/A");
+				InOutTooltip.AddTitle(TimerName);
 
-		const double TooltipEventDuration = TooltipEvent.GetDuration();
+				const double TooltipEventDuration = TooltipEvent.GetDuration();
 
-		if (TooltipEvent.GetDepth() > 0 &&
-			ParentTimingEvent.IsValid() &&
-			ParentTimingEvent->GetDuration() > 0.0 &&
-			ParentTimingEvent->GetDuration() != std::numeric_limits<double>::infinity())
-		{
-			const TraceServices::FTimingProfilerTimer* ParentTimer = TimerReader->GetTimer(ParentTimingEvent->GetTimerIndex());
-			const TCHAR* ParentTimerName = (ParentTimer != nullptr) ? ParentTimer->Name : TEXT("N/A");
-			FNumberFormattingOptions FormattingOptions;
-			FormattingOptions.MaximumFractionalDigits = 2;
-			const FString ValueStr = FString::Printf(TEXT("%s %s"), *FText::AsPercent(TooltipEventDuration / ParentTimingEvent->GetDuration(), &FormattingOptions).ToString(), ParentTimerName);
-			InOutTooltip.AddNameValueTextLine(TEXTVIEW("% of Parent:"), ValueStr);
-		}
+				if (TooltipEvent.GetDepth() > 0 &&
+					ParentTimingEvent.IsValid() &&
+					ParentTimingEvent->GetDuration() > 0.0 &&
+					ParentTimingEvent->GetDuration() != std::numeric_limits<double>::infinity())
+				{
+					const TraceServices::FTimingProfilerTimer* ParentTimer = TimerReader.GetTimer(ParentTimingEvent->GetTimerIndex());
+					const TCHAR* ParentTimerName = (ParentTimer != nullptr) ? ParentTimer->Name : TEXT("N/A");
+					FNumberFormattingOptions FormattingOptions;
+					FormattingOptions.MaximumFractionalDigits = 2;
+					const FString ValueStr = FString::Printf(TEXT("%s %s"), *FText::AsPercent(TooltipEventDuration / ParentTimingEvent->GetDuration(), &FormattingOptions).ToString(), ParentTimerName);
+					InOutTooltip.AddNameValueTextLine(TEXTVIEW("% of Parent:"), ValueStr);
+				}
 
-		if (TooltipEvent.GetDepth() > 1 &&
-			RootTimingEvent.IsValid() &&
-			RootTimingEvent->GetDuration() > 0.0 &&
-			RootTimingEvent->GetDuration() != std::numeric_limits<double>::infinity())
-		{
-			const TraceServices::FTimingProfilerTimer* RootTimer = TimerReader->GetTimer(RootTimingEvent->GetTimerIndex());
-			const TCHAR* RootTimerName = (RootTimer != nullptr) ? RootTimer->Name : TEXT("N/A");
-			FNumberFormattingOptions FormattingOptions;
-			FormattingOptions.MaximumFractionalDigits = 2;
-			const FString ValueStr = FString::Printf(TEXT("%s %s"), *FText::AsPercent(TooltipEventDuration / RootTimingEvent->GetDuration(), &FormattingOptions).ToString(), RootTimerName);
-			InOutTooltip.AddNameValueTextLine(TEXTVIEW("% of Root:"), ValueStr);
-		}
+				if (TooltipEvent.GetDepth() > 1 &&
+					RootTimingEvent.IsValid() &&
+					RootTimingEvent->GetDuration() > 0.0 &&
+					RootTimingEvent->GetDuration() != std::numeric_limits<double>::infinity())
+				{
+					const TraceServices::FTimingProfilerTimer* RootTimer = TimerReader.GetTimer(RootTimingEvent->GetTimerIndex());
+					const TCHAR* RootTimerName = (RootTimer != nullptr) ? RootTimer->Name : TEXT("N/A");
+					FNumberFormattingOptions FormattingOptions;
+					FormattingOptions.MaximumFractionalDigits = 2;
+					const FString ValueStr = FString::Printf(TEXT("%s %s"), *FText::AsPercent(TooltipEventDuration / RootTimingEvent->GetDuration(), &FormattingOptions).ToString(), RootTimerName);
+					InOutTooltip.AddNameValueTextLine(TEXTVIEW("% of Root:"), ValueStr);
+				}
 
-		InOutTooltip.AddNameValueTextLine(TEXTVIEW("Inclusive Time:"), FormatTimeAuto(TooltipEventDuration, 2));
+				InOutTooltip.AddNameValueTextLine(TEXTVIEW("Inclusive Time:"), FormatTimeAuto(TooltipEventDuration, 2));
 
-		if (TooltipEventDuration > 0.0 &&
-			TooltipEventDuration != std::numeric_limits<double>::infinity())
-		{
-			const double ExclusiveTimePercent = TooltipEvent.GetExclusiveTime() / TooltipEventDuration;
-			FNumberFormattingOptions FormattingOptions;
-			FormattingOptions.MaximumFractionalDigits = 2;
-			const FString ExclStr = FString::Printf(TEXT("%s (%s)"), *FormatTimeAuto(TooltipEvent.GetExclusiveTime(), 2), *FText::AsPercent(ExclusiveTimePercent, &FormattingOptions).ToString());
-			InOutTooltip.AddNameValueTextLine(TEXTVIEW("Exclusive Time:"), ExclStr);
-		}
+				if (TooltipEventDuration > 0.0 &&
+					TooltipEventDuration != std::numeric_limits<double>::infinity())
+				{
+					const double ExclusiveTimePercent = TooltipEvent.GetExclusiveTime() / TooltipEventDuration;
+					FNumberFormattingOptions FormattingOptions;
+					FormattingOptions.MaximumFractionalDigits = 2;
+					const FString ExclStr = FString::Printf(TEXT("%s (%s)"), *FormatTimeAuto(TooltipEvent.GetExclusiveTime(), 2), *FText::AsPercent(ExclusiveTimePercent, &FormattingOptions).ToString());
+					InOutTooltip.AddNameValueTextLine(TEXTVIEW("Exclusive Time:"), ExclStr);
+				}
 
-		InOutTooltip.AddNameValueTextLine(TEXTVIEW("Depth:"), FString::Printf(TEXT("%d"), TooltipEvent.GetDepth()));
+				InOutTooltip.AddNameValueTextLine(TEXTVIEW("Depth:"), FString::Printf(TEXT("%d"), TooltipEvent.GetDepth()));
 
-		const TraceServices::FMetadataSpec* MetadataSpec = nullptr;
-		if (Timer && Timer->HasValidMetadataSpecId())
-		{
-			MetadataSpec = TimingProfilerProvider.GetMetadataSpec(Timer->MetadataSpecId);
-		}
-		TArrayView<const uint8> Metadata = TimerReader->GetMetadata(TooltipEvent.GetTimerIndex());
-		if (Metadata.Num() > 0)
-		{
-			AppendMetadataToTooltip(InOutTooltip, MetadataSpec, Metadata);
-		}
+				const TraceServices::FMetadataSpec* MetadataSpec = nullptr;
+				if (Timer && Timer->HasValidMetadataSpecId())
+				{
+					MetadataSpec = TimingProfilerProvider.GetMetadataSpec(Timer->MetadataSpecId);
+				}
+				TArrayView<const uint8> Metadata = TimerReader.GetMetadata(TooltipEvent.GetTimerIndex());
+				if (Metadata.Num() > 0)
+				{
+					AppendMetadataToTooltip(InOutTooltip, MetadataSpec, Metadata);
+				}
 
-		PostInitTooltip(InOutTooltip, TooltipEvent, *Session.Get(), TimerName);
+				TSharedPtr<const TraceServices::IAnalysisSession> Session = FInsightsManager::Get()->GetSession();
+				if (Session.IsValid())
+				{
+					PostInitTooltip(InOutTooltip, TooltipEvent, *Session.Get(), TimerName);
+				}
+			}
+		); // ReadTimers
 	}
 	else
 	{
@@ -953,33 +956,22 @@ void FThreadTimingTrackImpl::GetParentAndRoot(const FThreadTrackEvent& TimingEve
 {
 	if (TimingEvent.GetDepth() > 0)
 	{
-		TSharedPtr<const TraceServices::IAnalysisSession> Session = FInsightsManager::Get()->GetSession();
-		if (Session.IsValid())
-		{
-			TraceServices::FAnalysisSessionReadScope SessionReadScope(*Session.Get());
-
-			if (TraceServices::ReadTimingProfilerProvider(*Session.Get()))
+		ReadTimeline(
+			[this, &TimingEvent, &OutParentTimingEvent, &OutRootTimingEvent]
+			(const TraceServices::ITimingProfilerProvider::Timeline& Timeline)
 			{
-				const TraceServices::ITimingProfilerProvider& TimingProfilerProvider = *TraceServices::ReadTimingProfilerProvider(*Session.Get());
-
-				TimingProfilerProvider.ReadTimeline(GetTimelineIndex(), [&TimingEvent, &OutParentTimingEvent, &OutRootTimingEvent](const TraceServices::ITimingProfilerProvider::Timeline& Timeline)
-					{
-						const double Time = (TimingEvent.GetStartTime() + TimingEvent.GetEndTime()) / 2;
-						TimelineEventInfo EventInfo;
-						bool IsFound = Timeline.GetEventInfo(Time, 0, TimingEvent.GetDepth() - 1, EventInfo);
-						if (IsFound)
-						{
-							CreateFThreadTrackEventFromInfo(EventInfo, TimingEvent.GetTrack(), TimingEvent.GetDepth() - 1, OutParentTimingEvent);
-						}
-
-						IsFound = Timeline.GetEventInfo(Time, 0, 0, EventInfo);
-						if (IsFound)
-						{
-							CreateFThreadTrackEventFromInfo(EventInfo, TimingEvent.GetTrack(), 0, OutRootTimingEvent);
-						}
-					});
+				const double Time = (TimingEvent.GetStartTime() + TimingEvent.GetEndTime()) / 2;
+				TimelineEventInfo EventInfo;
+				if (Timeline.GetEventInfo(Time, 0, TimingEvent.GetDepth() - 1, EventInfo))
+				{
+					OutParentTimingEvent = CreateEvent(EventInfo, TimingEvent.GetTrack(), TimingEvent.GetDepth() - 1);
+				}
+				if (Timeline.GetEventInfo(Time, 0, 0, EventInfo))
+				{
+					OutRootTimingEvent = CreateEvent(EventInfo, TimingEvent.GetTrack(), 0);
+				}
 			}
-		}
+		); // ReadTimeline
 	}
 }
 
@@ -1018,25 +1010,27 @@ const TSharedPtr<const ITimingEvent> FThreadTimingTrackImpl::GetEvent(float InPo
 
 		const double EventTime = Viewport.SlateUnitsToTime(InPosX);
 
+		double SessionDuration = std::numeric_limits<double>::lowest();
 		TSharedPtr<const TraceServices::IAnalysisSession> Session = FInsightsManager::Get()->GetSession();
 		if (Session.IsValid())
 		{
 			TraceServices::FAnalysisSessionReadScope SessionReadScope(*Session.Get());
+			SessionDuration = Session->GetDurationSeconds();
+		}
 
-			if (EventTime <= Session->GetDurationSeconds() && TraceServices::ReadTimingProfilerProvider(*Session.Get()))
-			{
-				const TraceServices::ITimingProfilerProvider& TimingProfilerProvider = *TraceServices::ReadTimingProfilerProvider(*Session.Get());
-
-				TimingProfilerProvider.ReadTimeline(GetTimelineIndex(), [this, &EventTime, &Depth, &TimingEvent, &SecondsPerPixel](const TraceServices::ITimingProfilerProvider::Timeline& Timeline)
+		if (EventTime <= SessionDuration)
+		{
+			ReadTimeline(
+				[this, &EventTime, &Depth, &TimingEvent, &SecondsPerPixel]
+				(const TraceServices::ITimingProfilerProvider::Timeline& Timeline)
+				{
+					TimelineEventInfo EventInfo;
+					if (Timeline.GetEventInfo(EventTime, 2 * SecondsPerPixel, Depth, EventInfo))
 					{
-						TimelineEventInfo EventInfo;
-						bool IsFound = Timeline.GetEventInfo(EventTime, 2 * SecondsPerPixel, Depth, EventInfo);
-						if (IsFound)
-						{
-							CreateFThreadTrackEventFromInfo(EventInfo, SharedThis(this), Depth, TimingEvent);
-						}
-					});
-			}
+						TimingEvent = CreateEvent(EventInfo, SharedThis(this), Depth);
+					}
+				}
+			); // ReadTimeline
 		}
 	}
 
@@ -1054,8 +1048,8 @@ const TSharedPtr<const ITimingEvent> FThreadTimingTrackImpl::SearchEvent(const F
 		FoundEvent->SetTimerIndex(InFoundEvent.TimerIndex);
 
 		uint32 TimerId = 0;
-		bool ret = FThreadTimingTrackImpl::TimerIndexToTimerId(InFoundEvent.TimerIndex, TimerId);
-		if (ret)
+		bool bIsValidTimerId = TimerIndexToTimerId(InFoundEvent.TimerIndex, TimerId);
+		if (bIsValidTimerId)
 		{
 			FoundEvent->SetTimerId(TimerId);
 		}
@@ -1071,22 +1065,11 @@ void FThreadTimingTrackImpl::UpdateEventStats(ITimingEvent& InOutEvent) const
 	if (InOutEvent.CheckTrack(this) && InOutEvent.Is<FThreadTrackEvent>())
 	{
 		FThreadTrackEvent& TrackEvent = InOutEvent.As<FThreadTrackEvent>();
-		if (TrackEvent.IsExclusiveTimeComputed())
-		{
-			return;
-		}
-
-		TSharedPtr<const TraceServices::IAnalysisSession> Session = FInsightsManager::Get()->GetSession();
-		if (Session.IsValid())
+		if (!TrackEvent.IsExclusiveTimeComputed())
 		{
-			TraceServices::FAnalysisSessionReadScope SessionReadScope(*Session.Get());
-
-			if (TraceServices::ReadTimingProfilerProvider(*Session.Get()))
-			{
-				const TraceServices::ITimingProfilerProvider& TimingProfilerProvider = *TraceServices::ReadTimingProfilerProvider(*Session.Get());
-
-				// Get Exclusive Time.
-				TimingProfilerProvider.ReadTimeline(GetTimelineIndex(), [&TrackEvent](const TraceServices::ITimingProfilerProvider::Timeline& Timeline)
+			ReadTimeline(
+				[&TrackEvent]
+				(const TraceServices::ITimingProfilerProvider::Timeline& Timeline)
 				{
 					TimelineEventInfo EventInfo;
 					bool bIsFound = Timeline.GetEventInfo(TrackEvent.GetStartTime(), 0.0, TrackEvent.GetDepth(), EventInfo);
@@ -1095,8 +1078,8 @@ void FThreadTimingTrackImpl::UpdateEventStats(ITimingEvent& InOutEvent) const
 						TrackEvent.SetExclusiveTime(EventInfo.ExclTime);
 						TrackEvent.SetIsExclusiveTimeComputed(true);
 					}
-				});
-			}
+				}
+			); // ReadTimeline
 		}
 	}
 }
@@ -1122,48 +1105,44 @@ void FThreadTimingTrackImpl::OnClipboardCopyEvent(const ITimingEvent& InSelected
 	{
 		const FThreadTrackEvent& TrackEvent = InSelectedEvent.As<FThreadTrackEvent>();
 
-		TSharedPtr<const TraceServices::IAnalysisSession> Session = FInsightsManager::Get()->GetSession();
-		check(Session.IsValid());
-
-		TraceServices::FAnalysisSessionReadScope SessionReadScope(*Session.Get());
-		const TraceServices::ITimingProfilerProvider& TimingProfilerProvider = *TraceServices::ReadTimingProfilerProvider(*Session.Get());
-		const TraceServices::ITimingProfilerTimerReader* TimerReader = nullptr;
-		TimingProfilerProvider.ReadTimers([&TimerReader](const TraceServices::ITimingProfilerTimerReader& Out) { TimerReader = &Out; });
-		check(TimerReader);
-
-		const TraceServices::FTimingProfilerTimer* TimerPtr = TimerReader->GetTimer(TrackEvent.GetTimerIndex());
+		ReadTimers(
+			[this, &TrackEvent](const TraceServices::ITimingProfilerProvider& TimingProfilerProvider, const TraceServices::ITimingProfilerTimerReader& TimerReader)
+			{
+				const TraceServices::FTimingProfilerTimer* TimerPtr = TimerReader.GetTimer(TrackEvent.GetTimerIndex());
 
-		if (TimerPtr)
-		{
-			FString EventName(TimerPtr->Name);
+				if (TimerPtr)
+				{
+					FString EventName(TimerPtr->Name);
 
-			FTimingEventsTrackDrawStateBuilder::AppendDurationToEventName(EventName, TrackEvent.GetDuration());
+					FTimingEventsTrackDrawStateBuilder::AppendDurationToEventName(EventName, TrackEvent.GetDuration());
 
-			const uint32 TimerIndex = TrackEvent.GetTimerIndex();
-			if (int32(TimerIndex) < 0) // has metadata?
-			{
-				TArrayView<const uint8> Metadata = TimerReader->GetMetadata(TimerIndex);
-				const TraceServices::FMetadataSpec* MetadataSpec = nullptr;
-				if (TimerPtr && TimerPtr->HasValidMetadataSpecId())
-				{
-					MetadataSpec = TimingProfilerProvider.GetMetadataSpec(TimerPtr->MetadataSpecId);
-				}
-				if (Metadata.Num() > 0)
-				{
-					if (MetadataSpec)
+					const uint32 TimerIndex = TrackEvent.GetTimerIndex();
+					if (int32(TimerIndex) < 0) // has metadata?
 					{
-						AppendMetadataToString(EventName, MetadataSpec, Metadata);
-					}
-					else
-					{
-						AppendMetadataToString(EventName, Metadata);
+						TArrayView<const uint8> Metadata = TimerReader.GetMetadata(TimerIndex);
+						const TraceServices::FMetadataSpec* MetadataSpec = nullptr;
+						if (TimerPtr && TimerPtr->HasValidMetadataSpecId())
+						{
+							MetadataSpec = TimingProfilerProvider.GetMetadataSpec(TimerPtr->MetadataSpecId);
+						}
+						if (Metadata.Num() > 0)
+						{
+							if (MetadataSpec)
+							{
+								AppendMetadataToString(EventName, MetadataSpec, Metadata);
+							}
+							else
+							{
+								AppendMetadataToString(EventName, Metadata);
+							}
+						}
 					}
+
+					// Copy name of selected timing event to clipboard.
+					FPlatformApplicationMisc::ClipboardCopy(*EventName);
 				}
 			}
-
-			// Copy name of selected timing event to clipboard.
-			FPlatformApplicationMisc::ClipboardCopy(*EventName);
-		}
+		); // ReadTimers
 	}
 }
 
@@ -1235,128 +1214,120 @@ bool FThreadTimingTrackImpl::FindTimingProfilerEvent(const FTimingEventSearchPar
 	FilterConfiguratorContext.AddFilterData<int64>(static_cast<int32>(EFilterField::Metadata), 0);
 
 	return TTimingEventSearch<TraceServices::FTimingProfilerEvent>::Search(
+
 		InParameters,
 
-		[this](TTimingEventSearch<TraceServices::FTimingProfilerEvent>::FContext& InContext)
+		// Search Predicate
+		[this]
+		(TTimingEventSearch<TraceServices::FTimingProfilerEvent>::FContext& InContext)
 		{
-			TSharedPtr<const TraceServices::IAnalysisSession> Session = FInsightsManager::Get()->GetSession();
-			if (Session.IsValid())
-			{
-				TraceServices::FAnalysisSessionReadScope SessionReadScope(*Session.Get());
-
-				if (TraceServices::ReadTimingProfilerProvider(*Session.Get()))
+			ReadTimeline(
+				[&InContext]
+				(const TraceServices::ITimingProfilerProvider::Timeline& Timeline)
 				{
-					const TraceServices::ITimingProfilerProvider& TimingProfilerProvider = *TraceServices::ReadTimingProfilerProvider(*Session.Get());
-
-					TimingProfilerProvider.ReadTimeline(GetTimelineIndex(), [&InContext](const TraceServices::ITimingProfilerProvider::Timeline& Timeline)
+					auto Callback = [&InContext](double EventStartTime, double EventEndTime, uint32 EventDepth, const TraceServices::FTimingProfilerEvent& Event)
 					{
-						auto Callback = [&InContext](double EventStartTime, double EventEndTime, uint32 EventDepth, const TraceServices::FTimingProfilerEvent& Event)
-						{
-							InContext.Check(EventStartTime, EventEndTime, EventDepth, Event);
-							return InContext.ShouldContinueSearching() ? TraceServices::EEventEnumerate::Continue : TraceServices::EEventEnumerate::Stop;
-						};
+						InContext.Check(EventStartTime, EventEndTime, EventDepth, Event);
+						return InContext.ShouldContinueSearching() ? TraceServices::EEventEnumerate::Continue : TraceServices::EEventEnumerate::Stop;
+					};
 
-						if (InContext.GetParameters().SearchDirection == FTimingEventSearchParameters::ESearchDirection::Forward)
-						{
-							Timeline.EnumerateEvents(InContext.GetParameters().StartTime, InContext.GetParameters().EndTime, Callback);
-						}
-						else
-						{
-							Timeline.EnumerateEventsBackwards(InContext.GetParameters().EndTime, InContext.GetParameters().StartTime, Callback);
-						}
-					});
+					if (InContext.GetParameters().SearchDirection == FTimingEventSearchParameters::ESearchDirection::Forward)
+					{
+						Timeline.EnumerateEvents(InContext.GetParameters().StartTime, InContext.GetParameters().EndTime, Callback);
+					}
+					else
+					{
+						Timeline.EnumerateEventsBackwards(InContext.GetParameters().EndTime, InContext.GetParameters().StartTime, Callback);
+					}
 				}
-			}
+			); // ReadTimeline
 		},
 
-		[&FilterConfiguratorContext, &InParameters](double EventStartTime, double EventEndTime, uint32 EventDepth, const TraceServices::FTimingProfilerEvent& Event)
+		// Filter Predicate
+		[this, &FilterConfiguratorContext, &InParameters]
+		(double EventStartTime, double EventEndTime, uint32 EventDepth, const TraceServices::FTimingProfilerEvent& Event)
 		{
 			if (!InParameters.FilterExecutor.IsValid())
 			{
 				return true;
 			}
 
-			TSharedPtr<const TraceServices::IAnalysisSession> Session = FInsightsManager::Get()->GetSession();
-			if (Session.IsValid())
-			{
-				TraceServices::FAnalysisSessionReadScope SessionReadScope(*Session.Get());
-				if (TraceServices::ReadTimingProfilerProvider(*Session.Get()))
-				{
-					const TraceServices::ITimingProfilerProvider& TimingProfilerProvider = *TraceServices::ReadTimingProfilerProvider(*Session.Get());
-					const TraceServices::ITimingProfilerTimerReader* TimerReader;
-					TimingProfilerProvider.ReadTimers([&TimerReader](const TraceServices::ITimingProfilerTimerReader& Out) { TimerReader = &Out; });
+			FilterConfiguratorContext.SetFilterData<double>(static_cast<int32>(EFilterField::StartTime), EventStartTime);
+			FilterConfiguratorContext.SetFilterData<double>(static_cast<int32>(EFilterField::EndTime), EventEndTime);
+			FilterConfiguratorContext.SetFilterData<double>(static_cast<int32>(EFilterField::Duration), EventEndTime - EventStartTime);
 
-					const TraceServices::FTimingProfilerTimer* Timer = TimerReader->GetTimer(Event.TimerIndex);
+			uint32 TimerId = ~0u;
+			ReadTimers(
+				[&Event, &TimerId]
+				(const TraceServices::ITimingProfilerProvider& TimingProfilerProvider, const TraceServices::ITimingProfilerTimerReader& TimerReader)
+				{
+					const TraceServices::FTimingProfilerTimer* Timer = TimerReader.GetTimer(Event.TimerIndex);
 					if (ensure(Timer != nullptr))
 					{
-						FilterConfiguratorContext.SetFilterData<double>(static_cast<int32>(EFilterField::StartTime), EventStartTime);
-						FilterConfiguratorContext.SetFilterData<double>(static_cast<int32>(EFilterField::EndTime), EventEndTime);
-						FilterConfiguratorContext.SetFilterData<double>(static_cast<int32>(EFilterField::Duration), EventEndTime - EventStartTime);
-
-						FilterConfiguratorContext.SetFilterData<int64>(static_cast<int32>(EFilterField::TimerId), Timer->Id);
-						FilterConfiguratorContext.SetFilterData<int64>(static_cast<int32>(EFilterField::TimerName), Timer->Id);
-
-						FilterConfiguratorContext.SetFilterData<int64>(static_cast<int32>(EFilterField::Metadata), Event.TimerIndex);
-
-						return InParameters.FilterExecutor->ApplyFilters(FilterConfiguratorContext);
+						TimerId = Timer->Id;
 					}
-				}
-			}
+				});
+			FilterConfiguratorContext.SetFilterData<int64>(static_cast<int32>(EFilterField::TimerId), TimerId);
+			FilterConfiguratorContext.SetFilterData<int64>(static_cast<int32>(EFilterField::TimerName), TimerId);
 
-			return false;
-		},
+			FilterConfiguratorContext.SetFilterData<int64>(static_cast<int32>(EFilterField::Metadata), Event.TimerIndex);
 
-		[&InFoundPredicate](double InFoundStartTime, double InFoundEndTime, uint32 InFoundDepth, const TraceServices::FTimingProfilerEvent& InEvent)
-		{
-			InFoundPredicate(InFoundStartTime, InFoundEndTime, InFoundDepth, InEvent);
+			return InParameters.FilterExecutor->ApplyFilters(FilterConfiguratorContext);
 		},
 
+		// Found Predicate
+		InFoundPredicate,
+
+		// Payload Matched Predicate
 		TTimingEventSearch<TraceServices::FTimingProfilerEvent>::NoMatch,
 
-		&SearchCache);
+		// Cache
+		&SearchCache
+
+	); // Search
 }
 
 ////////////////////////////////////////////////////////////////////////////////////////////////////
 
-void FThreadTimingTrackImpl::CreateFThreadTrackEventFromInfo(const TimelineEventInfo& InEventInfo, const TSharedRef<const FBaseTimingTrack> InTrack, int32 InDepth, TSharedPtr<FThreadTrackEvent> &OutTimingEvent)
+TSharedRef<FThreadTrackEvent> FThreadTimingTrackImpl::CreateEvent(const TimelineEventInfo& InEventInfo, const TSharedRef<const FBaseTimingTrack> InTrack, int32 InDepth) const
 {
-	OutTimingEvent = MakeShared<FThreadTrackEvent>(InTrack, InEventInfo.StartTime, InEventInfo.EndTime, InDepth);
-	FThreadTrackEvent& Event = OutTimingEvent->As<FThreadTrackEvent>();
+	TSharedRef<FThreadTrackEvent> ThreadTrackEvent = MakeShared<FThreadTrackEvent>(InTrack, InEventInfo.StartTime, InEventInfo.EndTime, InDepth);
+	FThreadTrackEvent& Event = *ThreadTrackEvent;
+
 	Event.SetExclusiveTime(InEventInfo.ExclTime);
 	Event.SetIsExclusiveTimeComputed(true);
 	Event.SetTimerIndex(InEventInfo.Event.TimerIndex);
 
 	uint32 TimerId = 0;
-	bool ret = FThreadTimingTrackImpl::TimerIndexToTimerId(InEventInfo.Event.TimerIndex, TimerId);
-	if (ret)
+	bool bIsValidTimerId = TimerIndexToTimerId(InEventInfo.Event.TimerIndex, TimerId);
+	if (bIsValidTimerId)
 	{
 		Event.SetTimerId(TimerId);
 	}
 
+	return ThreadTrackEvent;
 }
 
 ////////////////////////////////////////////////////////////////////////////////////////////////////
 
-bool FThreadTimingTrackImpl::TimerIndexToTimerId(uint32 InTimerIndex, uint32& OutTimerId)
+bool FThreadTimingTrackImpl::TimerIndexToTimerId(uint32 InTimerIndex, uint32& OutTimerId) const
 {
-	TSharedPtr<const TraceServices::IAnalysisSession> Session = FInsightsManager::Get()->GetSession();
-	check(Session.IsValid())
-
-	TraceServices::FAnalysisSessionReadScope SessionReadScope(*Session.Get());
-
-	const TraceServices::ITimingProfilerProvider& TimingProfilerProvider = *TraceServices::ReadTimingProfilerProvider(*Session.Get());
+	bool bIsValidTimerId = false;
 
-	const TraceServices::ITimingProfilerTimerReader* TimerReader;
-	TimingProfilerProvider.ReadTimers([&TimerReader](const TraceServices::ITimingProfilerTimerReader& Out) { TimerReader = &Out; });
-
-	const TraceServices::FTimingProfilerTimer* Timer = TimerReader->GetTimer(InTimerIndex);
-	if (Timer == nullptr)
-	{
-		return false;
-	}
+	ReadTimers(
+		[&bIsValidTimerId, InTimerIndex, &OutTimerId]
+		(const TraceServices::ITimingProfilerProvider& TimingProfilerProvider, const TraceServices::ITimingProfilerTimerReader& TimerReader)
+		{
+			const TraceServices::FTimingProfilerTimer* Timer = TimerReader.GetTimer(InTimerIndex);
+			if (Timer)
+			{
+				OutTimerId = Timer->Id;
+				bIsValidTimerId = true;
+			}
+		}
+	); // ReadTimers
 
-	OutTimerId = Timer->Id;
-	return true;
+	return bIsValidTimerId;
 }
 
 ////////////////////////////////////////////////////////////////////////////////////////////////////
@@ -1371,19 +1342,15 @@ bool FThreadTimingTrackImpl::HasCustomFilter() const
 int32 FThreadTimingTrackImpl::GetDepthAt(double Time) const
 {
 	int32 Depth = 0;
-	TSharedPtr<const TraceServices::IAnalysisSession> Session = FInsightsManager::Get()->GetSession();
-	if (Session.IsValid() && TraceServices::ReadTimingProfilerProvider(*Session.Get()))
-	{
-		TraceServices::FAnalysisSessionReadScope SessionReadScope(*Session.Get());
 
-		const TraceServices::ITimingProfilerProvider& TimingProfilerProvider = *TraceServices::ReadTimingProfilerProvider(*Session.Get());
+	ReadTimeline(
+		[Time, &Depth]
+		(const TraceServices::ITimingProfilerProvider::Timeline& Timeline)
+		{
+			Depth = Timeline.GetDepthAt(Time);
+		}
+	); // ReadTimeline
 
-		TimingProfilerProvider.ReadTimeline(TimelineIndex,
-			[Time, &Depth](const TraceServices::ITimingProfilerProvider::Timeline& Timeline)
-			{
-				Depth = Timeline.GetDepthAt(Time);
-			});
-	}
 	return Depth;
 }
 
@@ -1402,81 +1369,81 @@ void FThreadTimingTrackImpl::SetFilterConfigurator(TSharedPtr<FFilterConfigurato
 
 TSharedPtr<const ITimingEvent> FThreadTimingTrackImpl::FindMaxEventInstance(uint32 InTimerId, double InStartTime, double InEndTime) const
 {
-	TSharedPtr<const TraceServices::IAnalysisSession> Session = FInsightsManager::Get()->GetSession();
-	check(Session.IsValid());
-
-	TraceServices::FAnalysisSessionReadScope SessionReadScope(*Session.Get());
+	TSharedPtr<FThreadTrackEvent> TimingEvent;
 
-	const TraceServices::ITimingProfilerProvider& TimingProfilerProvider = *TraceServices::ReadTimingProfilerProvider(*Session.Get());
+	ReadTimers(
+		[this, InTimerId, InStartTime, InEndTime, &TimingEvent]
+		(const TraceServices::ITimingProfilerProvider& TimingProfilerProvider, const TraceServices::ITimingProfilerTimerReader& TimerReader)
+		{
+			ReadTimeline(
+				[this, InTimerId, InStartTime, InEndTime, &TimingEvent, &TimerReader]
+				(const TraceServices::ITimingProfilerProvider::Timeline& Timeline)
+				{
+					struct FCandidateEvent
+					{
+						double StartTime = 0.0f;
+						double EndTime = -1.0f;
+						uint32 Depth = 0;
+						uint32 TimerIndex = 0;
+					};
+					TArray<FCandidateEvent> CandidateEvents;
 
-	const TraceServices::ITimingProfilerTimerReader* TimerReader;
-	TimingProfilerProvider.ReadTimers([&TimerReader](const TraceServices::ITimingProfilerTimerReader& Out) { TimerReader = &Out; });
+					TraceServices::ITimeline<TraceServices::FTimingProfilerEvent>::EnumerateAsyncParams Params;
+					Params.IntervalStart = InStartTime;
+					Params.IntervalEnd = InEndTime;
+					Params.Resolution = 0.0;
+					Params.SetupCallback =
+						[&CandidateEvents](uint32 NumTasks)
+						{
+							CandidateEvents.AddDefaulted(NumTasks);
+						};
+					Params.EventRangeCallback =
+						[InTimerId, &TimerReader, &CandidateEvents]
+						(double StartTime, double EndTime, uint32 Depth, const TraceServices::FTimingProfilerEvent& Event, uint32 TaskIndex)
+						{
+							const TraceServices::FTimingProfilerTimer* Timer = TimerReader.GetTimer(Event.TimerIndex);
+							if (ensure(Timer != nullptr))
+							{
+								if (Timer->Id == InTimerId)
+								{
+									double CandidateDuration = CandidateEvents[TaskIndex].EndTime - CandidateEvents[TaskIndex].StartTime;
+									double EventDuration = EndTime - StartTime;
 
-	struct CandidateEvent
-	{
-		double StartTime = 0.0f;
-		double EndTime = -1.0f;
-		uint32 Depth = 0;
-		uint32 TimerIndex = 0;
-	};
+									if (EventDuration > CandidateDuration)
+									{
+										CandidateEvents[TaskIndex].StartTime = StartTime;
+										CandidateEvents[TaskIndex].EndTime = EndTime;
+										CandidateEvents[TaskIndex].Depth = Depth;
+										CandidateEvents[TaskIndex].TimerIndex = Event.TimerIndex;
+									}
 
-	TSharedPtr<FThreadTrackEvent> TimingEvent;
+								}
+							}
+							return TraceServices::EEventEnumerate::Continue;
+						};
 
-	TimingProfilerProvider.ReadTimeline(TimelineIndex,
-		[TimerReader, InStartTime, InEndTime, InTimerId, &TimingEvent, this](const TraceServices::ITimingProfilerProvider::Timeline& Timeline)
-		{
-			TArray<CandidateEvent> CandidateEvents;
+					// Note: Enumerating events for filtering should not use downsampling.
+					Timeline.EnumerateEventsDownSampledAsync(Params);
 
-			TraceServices::ITimeline<TraceServices::FTimingProfilerEvent>::EnumerateAsyncParams Params;
-			Params.IntervalStart = InStartTime;
-			Params.IntervalEnd = InEndTime;
-			Params.Resolution = 0.0;
-			Params.SetupCallback = [&CandidateEvents](uint32 NumTasks)
-			{
-				CandidateEvents.AddDefaulted(NumTasks);
-			};
-			Params.EventRangeCallback = [TimerReader, &CandidateEvents, InTimerId](double StartTime, double EndTime, uint32 Depth, const TraceServices::FTimingProfilerEvent& Event, uint32 TaskIndex)
-			{
-				const TraceServices::FTimingProfilerTimer* Timer = TimerReader->GetTimer(Event.TimerIndex);
-				if (ensure(Timer != nullptr))
-				{
-					if (Timer->Id == InTimerId)
+					FCandidateEvent BestMatch;
+					for (const FCandidateEvent& Event : CandidateEvents)
 					{
-						double CandidateDuration = CandidateEvents[TaskIndex].EndTime - CandidateEvents[TaskIndex].StartTime;
-						double EventDuration = EndTime - StartTime;
-
-						if (EventDuration > CandidateDuration)
+						if ((Event.EndTime - Event.StartTime) > BestMatch.EndTime - BestMatch.StartTime)
 						{
-							CandidateEvents[TaskIndex].StartTime = StartTime;
-							CandidateEvents[TaskIndex].EndTime = EndTime;
-							CandidateEvents[TaskIndex].Depth = Depth;
-							CandidateEvents[TaskIndex].TimerIndex = Event.TimerIndex;
+							BestMatch = Event;
 						}
-
 					}
-				}
-				return TraceServices::EEventEnumerate::Continue;
-			};
-
-			// Note: Enumerating events for filtering should not use downsampling.
-			Timeline.EnumerateEventsDownSampledAsync(Params);
 
-			CandidateEvent BestMatch;
-			for (const CandidateEvent& Event : CandidateEvents)
-			{
-				if ((Event.EndTime - Event.StartTime) > BestMatch.EndTime - BestMatch.StartTime)
-				{
-					BestMatch = Event;
+					if (BestMatch.EndTime > BestMatch.StartTime)
+					{
+						TimingEvent = MakeShared<FThreadTrackEvent>(SharedThis(this), BestMatch.StartTime, BestMatch.EndTime, BestMatch.Depth);
+						TimingEvent->SetTimerId(InTimerId);
+						TimingEvent->SetTimerIndex(BestMatch.TimerIndex);
+					}
 				}
-			}
-
-			if (BestMatch.EndTime > BestMatch.StartTime)
-			{
-				TimingEvent = MakeShared<FThreadTrackEvent>(SharedThis(this), BestMatch.StartTime, BestMatch.EndTime, BestMatch.Depth);
-				TimingEvent->SetTimerId(InTimerId);
-				TimingEvent->SetTimerIndex(BestMatch.TimerIndex);
-			}
-		});
+			); // ReadTimeline
+		}
+	); // ReadTimers
 
 	return TimingEvent;
 }
@@ -1485,82 +1452,128 @@ TSharedPtr<const ITimingEvent> FThreadTimingTrackImpl::FindMaxEventInstance(uint
 
 TSharedPtr<const ITimingEvent> FThreadTimingTrackImpl::FindMinEventInstance(uint32 InTimerId, double InStartTime, double InEndTime) const
 {
-	TSharedPtr<const TraceServices::IAnalysisSession> Session = FInsightsManager::Get()->GetSession();
-	check(Session.IsValid());
-
-	TraceServices::FAnalysisSessionReadScope SessionReadScope(*Session.Get());
-
-	const TraceServices::ITimingProfilerProvider& TimingProfilerProvider = *TraceServices::ReadTimingProfilerProvider(*Session.Get());
-
-	const TraceServices::ITimingProfilerTimerReader* TimerReader;
-	TimingProfilerProvider.ReadTimers([&TimerReader](const TraceServices::ITimingProfilerTimerReader& Out) { TimerReader = &Out; });
-
-	struct CandidateEvent
-	{
-		double StartTime = -std::numeric_limits<double>::infinity();
-		double EndTime = std::numeric_limits<double>::infinity();
-		uint32 Depth = 0;
-		uint32 TimerIndex = 0;
-	};
-
 	TSharedPtr<FThreadTrackEvent> TimingEvent;
 
-	TimingProfilerProvider.ReadTimeline(TimelineIndex,
-		[TimerReader, InStartTime, InEndTime, InTimerId, &TimingEvent, this](const TraceServices::ITimingProfilerProvider::Timeline& Timeline)
+	ReadTimers(
+		[this, InTimerId, InStartTime, InEndTime, &TimingEvent]
+		(const TraceServices::ITimingProfilerProvider& TimingProfilerProvider, const TraceServices::ITimingProfilerTimerReader& TimerReader)
 		{
-			TArray<CandidateEvent> CandidateEvents;
-
-			TraceServices::ITimeline<TraceServices::FTimingProfilerEvent>::EnumerateAsyncParams Params;
-			Params.IntervalStart = InStartTime;
-			Params.IntervalEnd = InEndTime;
-			Params.Resolution = 0.0;
-			Params.SetupCallback = [&CandidateEvents](uint32 NumTasks)
-			{
-				CandidateEvents.AddDefaulted(NumTasks);
-			};
-			Params.EventRangeCallback = [TimerReader, &CandidateEvents, InTimerId](double StartTime, double EndTime, uint32 Depth, const TraceServices::FTimingProfilerEvent& Event, uint32 TaskIndex)
-			{
-				const TraceServices::FTimingProfilerTimer* Timer = TimerReader->GetTimer(Event.TimerIndex);
-				if (ensure(Timer != nullptr))
+			ReadTimeline(
+				[this, InTimerId, InStartTime, InEndTime, &TimingEvent, &TimerReader]
+				(const TraceServices::ITimingProfilerProvider::Timeline& Timeline)
 				{
-					if (Timer->Id == InTimerId)
+					struct FCandidateEvent
 					{
-						double CandidateDuration = CandidateEvents[TaskIndex].EndTime - CandidateEvents[TaskIndex].StartTime;
-						double EventDuration = EndTime - StartTime;
+						double StartTime = -std::numeric_limits<double>::infinity();
+						double EndTime = std::numeric_limits<double>::infinity();
+						uint32 Depth = 0;
+						uint32 TimerIndex = 0;
+					};
+					TArray<FCandidateEvent> CandidateEvents;
+
+					TraceServices::ITimeline<TraceServices::FTimingProfilerEvent>::EnumerateAsyncParams Params;
+					Params.IntervalStart = InStartTime;
+					Params.IntervalEnd = InEndTime;
+					Params.Resolution = 0.0;
+					Params.SetupCallback =
+						[&CandidateEvents](uint32 NumTasks)
+						{
+							CandidateEvents.AddDefaulted(NumTasks);
+						};
+					Params.EventRangeCallback =
+						[InTimerId, &TimerReader, &CandidateEvents]
+						(double StartTime, double EndTime, uint32 Depth, const TraceServices::FTimingProfilerEvent& Event, uint32 TaskIndex)
+						{
+							const TraceServices::FTimingProfilerTimer* Timer = TimerReader.GetTimer(Event.TimerIndex);
+							if (ensure(Timer != nullptr))
+							{
+								if (Timer->Id == InTimerId)
+								{
+									double CandidateDuration = CandidateEvents[TaskIndex].EndTime - CandidateEvents[TaskIndex].StartTime;
+									double EventDuration = EndTime - StartTime;
+
+									if (EventDuration < CandidateDuration)
+									{
+										CandidateEvents[TaskIndex].StartTime = StartTime;
+										CandidateEvents[TaskIndex].EndTime = EndTime;
+										CandidateEvents[TaskIndex].Depth = Depth;
+										CandidateEvents[TaskIndex].TimerIndex = Event.TimerIndex;
+									}
+								}
+							}
+							return TraceServices::EEventEnumerate::Continue;
+						};
+
+					// Note: Enumerating events for filtering should not use downsampling.
+					Timeline.EnumerateEventsDownSampledAsync(Params);
 
-						if (EventDuration < CandidateDuration)
+					FCandidateEvent BestMatch;
+					for (const FCandidateEvent& Event : CandidateEvents)
+					{
+						if ((Event.EndTime - Event.StartTime) < BestMatch.EndTime - BestMatch.StartTime)
 						{
-							CandidateEvents[TaskIndex].StartTime = StartTime;
-							CandidateEvents[TaskIndex].EndTime = EndTime;
-							CandidateEvents[TaskIndex].Depth = Depth;
-							CandidateEvents[TaskIndex].TimerIndex = Event.TimerIndex;
+							BestMatch = Event;
 						}
 					}
+
+					if (BestMatch.StartTime != -std::numeric_limits<double>::infinity())
+					{
+						TimingEvent = MakeShared<FThreadTrackEvent>(SharedThis(this), BestMatch.StartTime, BestMatch.EndTime, BestMatch.Depth);
+						TimingEvent->SetTimerId(InTimerId);
+						TimingEvent->SetTimerIndex(BestMatch.TimerIndex);
+					}
 				}
-				return TraceServices::EEventEnumerate::Continue;
-			};
+			); // ReadTimeline
+		}
+	); // ReadTimers
 
-			// Note: Enumerating events for filtering should not use downsampling.
-			Timeline.EnumerateEventsDownSampledAsync(Params);
+	return TimingEvent;
+}
 
-			CandidateEvent BestMatch;
-			for (const CandidateEvent& Event : CandidateEvents)
-			{
-				if ((Event.EndTime - Event.StartTime) < BestMatch.EndTime - BestMatch.StartTime)
+////////////////////////////////////////////////////////////////////////////////////////////////////
+
+bool FThreadTimingTrackImpl::ReadTimers(TFunctionRef<void(const TraceServices::ITimingProfilerProvider&, const TraceServices::ITimingProfilerTimerReader&)> Callback) const
+{
+	TSharedPtr<const TraceServices::IAnalysisSession> Session = FInsightsManager::Get()->GetSession();
+	if (Session.IsValid())
+	{
+		const TraceServices::ITimingProfilerProvider* TimingProfilerProvider = TraceServices::ReadTimingProfilerProvider(*Session.Get());
+		if (TimingProfilerProvider)
+		{
+			TraceServices::FAnalysisSessionReadScope SessionReadScope(*Session.Get());
+			TimingProfilerProvider->ReadTimers(
+				[&Callback, TimingProfilerProvider]
+				(const TraceServices::ITimingProfilerTimerReader& TimerReader)
 				{
-					BestMatch = Event;
-				}
-			}
+					Callback(*TimingProfilerProvider, TimerReader);
+				});
+			return true;
+		}
+	}
+	return false;
+}
 
-			if (BestMatch.StartTime != -std::numeric_limits<double>::infinity())
-			{
-				TimingEvent = MakeShared<FThreadTrackEvent>(SharedThis(this), BestMatch.StartTime, BestMatch.EndTime, BestMatch.Depth);
-				TimingEvent->SetTimerId(InTimerId);
-				TimingEvent->SetTimerIndex(BestMatch.TimerIndex);
-			}
-		});
+////////////////////////////////////////////////////////////////////////////////////////////////////
 
-	return TimingEvent;
+bool FThreadTimingTrackImpl::ReadTimeline(TFunctionRef<void(const TraceServices::ITimingProfilerProvider::Timeline&)> Callback) const
+{
+	TSharedPtr<const TraceServices::IAnalysisSession> Session = FInsightsManager::Get()->GetSession();
+	if (Session.IsValid())
+	{
+		const TraceServices::ITimingProfilerProvider* TimingProfilerProvider = TraceServices::ReadTimingProfilerProvider(*Session.Get());
+		if (TimingProfilerProvider)
+		{
+			TraceServices::FAnalysisSessionReadScope SessionReadScope(*Session.Get());
+			TimingProfilerProvider->ReadTimeline(GetTimelineIndex(),
+				[&Callback]
+				(const TraceServices::ITimingProfilerProvider::Timeline& Timeline)
+				{
+					Callback(Timeline);
+				});
+			return true;
+		}
+	}
+	return false;
 }
 
 ////////////////////////////////////////////////////////////////////////////////////////////////////
diff --git a/Engine/Source/Developer/TraceInsights/Private/Insights/TimingProfiler/Tracks/ThreadTimingTrackPrivate.h b/Engine/Source/Developer/TraceInsights/Private/Insights/TimingProfiler/Tracks/ThreadTimingTrackPrivate.h
--- a/Engine/Source/Developer/TraceInsights/Private/Insights/TimingProfiler/Tracks/ThreadTimingTrackPrivate.h
+++ b/Engine/Source/Developer/TraceInsights/Private/Insights/TimingProfiler/Tracks/ThreadTimingTrackPrivate.h
@@ -93,6 +93,9 @@ protected:
 
 	virtual void PostInitTooltip(FTooltipDrawState& InOutTooltip, const FThreadTrackEvent& TooltipEvent, const TraceServices::IAnalysisSession& Session, const TCHAR* TimerName) const {}
 
+	virtual bool ReadTimers(TFunctionRef<void(const TraceServices::ITimingProfilerProvider&, const TraceServices::ITimingProfilerTimerReader&)> Callback) const;
+	virtual bool ReadTimeline(TFunctionRef<void(const TraceServices::ITimingProfilerProvider::Timeline&)> Callback) const;
+
 private:
 	bool FindTimingProfilerEvent(const FThreadTrackEvent& InTimingEvent, TFunctionRef<void(double, double, uint32, const TraceServices::FTimingProfilerEvent&)> InFoundPredicate) const;
 	bool FindTimingProfilerEvent(const FTimingEventSearchParameters& InParameters, TFunctionRef<void(double, double, uint32, const TraceServices::FTimingProfilerEvent&)> InFoundPredicate) const;
@@ -101,8 +104,8 @@ private:
 						  TSharedPtr<FThreadTrackEvent>& OutParentTimingEvent,
 						  TSharedPtr<FThreadTrackEvent>& OutRootTimingEvent) const;
 
-	static void CreateFThreadTrackEventFromInfo(const TimelineEventInfo& InEventInfo, const TSharedRef<const FBaseTimingTrack> InTrack, int32 InDepth, TSharedPtr<FThreadTrackEvent> &OutTimingEvent);
-	static bool TimerIndexToTimerId(uint32 InTimerIndex, uint32 & OutTimerId);
+	TSharedRef<FThreadTrackEvent> CreateEvent(const TimelineEventInfo& InEventInfo, const TSharedRef<const FBaseTimingTrack> InTrack, int32 InDepth) const;
+	bool TimerIndexToTimerId(uint32 InTimerIndex, uint32& OutTimerId) const;
 
 private:
 	FThreadTimingSharedState& SharedState;
