[Insights] Cpu Timing Track: Added task priorities to applicable events regardless of thread priority

diff --git a/Engine/Source/Developer/TraceInsights/Private/Insights/TimingProfiler/Tracks/CpuTimingTrack.cpp b/Engine/Source/Developer/TraceInsights/Private/Insights/TimingProfiler/Tracks/CpuTimingTrack.cpp
--- a/Engine/Source/Developer/TraceInsights/Private/Insights/TimingProfiler/Tracks/CpuTimingTrack.cpp
+++ b/Engine/Source/Developer/TraceInsights/Private/Insights/TimingProfiler/Tracks/CpuTimingTrack.cpp
@@ -55,10 +55,10 @@ void FCpuTimingTrack::AddTaskInfo(FTooltipDrawState& InOutTooltip, const TraceSe
 		return Track.IsValid() ? Track->GetName() : TEXT("Unknown");
 	};
 
+	const TCHAR* TaskPri = ENamedThreads::GetTaskPriority(ThreadInfo) == ENamedThreads::NormalTaskPriority ? TEXT("Normal") : TEXT("High");
+
 	if (ThreadIndex == ENamedThreads::AnyThread)
 	{
-		const TCHAR* TaskPri = ENamedThreads::GetTaskPriority(ThreadInfo) == ENamedThreads::NormalTaskPriority ? TEXT("Normal") : TEXT("High");
-
 		int32 ThreadPriIndex = ENamedThreads::GetThreadPriorityIndex(ThreadInfo);
 		const TCHAR* ThreadPriStrs[] = { TEXT("Normal"), TEXT("High"), TEXT("Low") };
 		const TCHAR* ThreadPri = ensure(ThreadPriIndex >= 0 && ThreadPriIndex < 3) ? ThreadPriStrs[ThreadPriIndex] : TEXT("Unknown");
@@ -71,7 +71,7 @@ void FCpuTimingTrack::AddTaskInfo(FTooltipDrawState& InOutTooltip, const TraceSe
 	{
 		const TCHAR* QueueStr = ENamedThreads::GetQueueIndex(ThreadInfo) == ENamedThreads::MainQueue ? TEXT("Main") : TEXT("Local");
 		InOutTooltip.AddTextLine(
-			FString::Printf(TEXT("%s (%s queue)"), *GetTrackName(Task.StartedThreadId), QueueStr),
+			FString::Printf(TEXT("%s Pri task on %s (%s queue)"), TaskPri, *GetTrackName(Task.StartedThreadId), QueueStr),
 			FLinearColor::Green);
 	}
 
