[Insights] Added separate lock for the context switches provider.

#rnx

diff --git a/Engine/Source/Developer/TraceInsights/Private/Insights/ContextSwitches/ContextSwitchesProfilerManager.cpp b/Engine/Source/Developer/TraceInsights/Private/Insights/ContextSwitches/ContextSwitchesProfilerManager.cpp
--- a/Engine/Source/Developer/TraceInsights/Private/Insights/ContextSwitches/ContextSwitchesProfilerManager.cpp
+++ b/Engine/Source/Developer/TraceInsights/Private/Insights/ContextSwitches/ContextSwitchesProfilerManager.cpp
@@ -7,6 +7,7 @@
 #include "Modules/ModuleManager.h"
 
 // TraceServices
+#include "Common/ProviderLock.h"
 #include "TraceServices/Model/ContextSwitches.h"
 
 // TraceInsights
@@ -134,10 +135,13 @@ bool FContextSwitchesProfilerManager::Tick(float DeltaTime)
 		TSharedPtr<const TraceServices::IAnalysisSession> Session = FInsightsManager::Get()->GetSession();
 		if (Session.IsValid())
 		{
-			TraceServices::FAnalysisSessionReadScope SessionReadScope(*Session.Get());
-
-			const TraceServices::IContextSwitchesProvider* ContextSwitchesProvider = TraceServices::ReadContextSwitchesProvider(*Session.Get());
-			if (ContextSwitchesProvider && ContextSwitchesProvider->HasData())
+			bool bHasData = false;
+			if (const TraceServices::IContextSwitchesProvider* ContextSwitchesProvider = TraceServices::ReadContextSwitchesProvider(*Session.Get()))
+			{
+				TraceServices::FProviderReadScopeLock ScopedLock(*ContextSwitchesProvider);
+				bHasData = ContextSwitchesProvider->HasData();
+			}
+			if (bHasData)
 			{
 				using namespace UE::Insights::TimingProfiler;
 
diff --git a/Engine/Source/Developer/TraceInsights/Private/Insights/ContextSwitches/ViewModels/ContextSwitchesFilterConverters.cpp b/Engine/Source/Developer/TraceInsights/Private/Insights/ContextSwitches/ViewModels/ContextSwitchesFilterConverters.cpp
--- a/Engine/Source/Developer/TraceInsights/Private/Insights/ContextSwitches/ViewModels/ContextSwitchesFilterConverters.cpp
+++ b/Engine/Source/Developer/TraceInsights/Private/Insights/ContextSwitches/ViewModels/ContextSwitchesFilterConverters.cpp
@@ -3,6 +3,7 @@
 #include "ContextSwitchesFilterConverters.h"
 
 // TraceServices
+#include "Common/ProviderLock.h"
 #include "TraceServices/Model/ContextSwitches.h"
 #include "TraceServices/Model/Threads.h"
 
@@ -28,9 +29,10 @@ bool FCoreEventNameFilterValueConverter::Convert(const FString& Input, int64& Ou
 			{
 				if (!bFound && FCString::Stricmp(*Input, ThreadInfo.Name) == 0)
 				{
-					const TraceServices::IContextSwitchesProvider* ContextSwitchesProvider = TraceServices::ReadContextSwitchesProvider(*Session.Get());
-					if (ContextSwitchesProvider)
+					if (const TraceServices::IContextSwitchesProvider* ContextSwitchesProvider = TraceServices::ReadContextSwitchesProvider(*Session.Get()))
 					{
+						TraceServices::FProviderReadScopeLock ScopedLock(*ContextSwitchesProvider);
+
 						uint32 SystemThreadId;
 						if (ContextSwitchesProvider->GetSystemThreadId(ThreadInfo.Id, SystemThreadId))
 						{
diff --git a/Engine/Source/Developer/TraceInsights/Private/Insights/ContextSwitches/ViewModels/ContextSwitchesSharedState.cpp b/Engine/Source/Developer/TraceInsights/Private/Insights/ContextSwitches/ViewModels/ContextSwitchesSharedState.cpp
--- a/Engine/Source/Developer/TraceInsights/Private/Insights/ContextSwitches/ViewModels/ContextSwitchesSharedState.cpp
+++ b/Engine/Source/Developer/TraceInsights/Private/Insights/ContextSwitches/ViewModels/ContextSwitchesSharedState.cpp
@@ -7,6 +7,7 @@
 #include "Framework/MultiBox/MultiBoxBuilder.h"
 
 // TraceServices
+#include "Common/ProviderLock.h"
 #include "TraceServices/Model/Threads.h"
 #include "TraceServices/Model/ContextSwitches.h"
 
@@ -304,7 +305,7 @@ void FContextSwitchesSharedState::AddCommands()
 	FContextSwitchesStateCommands::Register();
 
 	TSharedPtr<TimingProfiler::STimingView> TimingView = GetTimingView();
-	
+
 	if (!TimingView.IsValid())
 	{
 		return;
@@ -387,37 +388,35 @@ void FContextSwitchesSharedState::AddCoreTracks()
 		return;
 	}
 
-	TraceServices::FAnalysisSessionReadScope SessionReadScope(*Session.Get());
-
-	const TraceServices::IContextSwitchesProvider* ContextSwitchesProvider = TraceServices::ReadContextSwitchesProvider(*Session.Get());
-	if (ContextSwitchesProvider == nullptr)
+	if (const TraceServices::IContextSwitchesProvider* ContextSwitchesProvider = TraceServices::ReadContextSwitchesProvider(*Session.Get()))
 	{
-		return;
-	}
+		TraceServices::FProviderReadScopeLock ScopedLock(*ContextSwitchesProvider);
 
-	const uint64 NewCpuCoresSerial = ContextSwitchesProvider->GetCpuCoresSerial();
-	if (NewCpuCoresSerial == CpuCoresSerial)
-	{
-		return;
-	}
-
-	CpuCoresSerial = NewCpuCoresSerial;
-
-	ContextSwitchesProvider->EnumerateCpuCores([this, TimingView](const TraceServices::FCpuCoreInfo& CpuCoreInfo)
+		const uint64 NewCpuCoresSerial = ContextSwitchesProvider->GetCpuCoresSerial();
+		if (NewCpuCoresSerial != CpuCoresSerial)
 		{
-			TSharedPtr<FCpuCoreTimingTrack>* TrackPtrPtr = CpuCoreTimingTracks.Find(CpuCoreInfo.CoreNumber);
-			if (TrackPtrPtr == nullptr)
-			{
-				const FString TrackName = FString::Printf(TEXT("Core %u"), CpuCoreInfo.CoreNumber);
-				TSharedPtr<FCpuCoreTimingTrack> Track = MakeShared<FCpuCoreTimingTrack>(*this, TrackName, CpuCoreInfo.CoreNumber);
+			CpuCoresSerial = NewCpuCoresSerial;
+
+			ContextSwitchesProvider->EnumerateCpuCores(
+				[this, TimingView]
+				(const TraceServices::FCpuCoreInfo& CpuCoreInfo)
+				{
+					TSharedPtr<FCpuCoreTimingTrack>* TrackPtrPtr = CpuCoreTimingTracks.Find(CpuCoreInfo.CoreNumber);
+					if (TrackPtrPtr == nullptr)
+					{
+						const FString TrackName = FString::Printf(TEXT("Core %u"), CpuCoreInfo.CoreNumber);
+						TSharedPtr<FCpuCoreTimingTrack> Track = MakeShared<FCpuCoreTimingTrack>(*this, TrackName, CpuCoreInfo.CoreNumber);
 
-				const int32 Order = FTimingTrackOrder::Cpu - 1024 + CpuCoreInfo.CoreNumber;
-				Track->SetOrder(Order);
+						const int32 Order = FTimingTrackOrder::Cpu - 1024 + CpuCoreInfo.CoreNumber;
+						Track->SetOrder(Order);
 
-				CpuCoreTimingTracks.Add(CpuCoreInfo.CoreNumber, Track);
-				TimingView->AddScrollableTrack(Track);
-			}
-		});
+						CpuCoreTimingTracks.Add(CpuCoreInfo.CoreNumber, Track);
+						TimingView->AddScrollableTrack(Track);
+					}
+				}
+			);
+		}
+	}
 }
 
 ////////////////////////////////////////////////////////////////////////////////////////////////////
@@ -442,32 +441,7 @@ void FContextSwitchesSharedState::RemoveCoreTracks()
 
 void FContextSwitchesSharedState::AddContextSwitchesChildTracks()
 {
-#if 0
-	TSharedPtr<const TraceServices::IAnalysisSession> Session = FInsightsManager::Get()->GetSession();
-	if (!Session.IsValid())
-	{
-		return;
-	}
-
-	TraceServices::FAnalysisSessionReadScope SessionReadScope(*Session.Get());
-
-	const TraceServices::IContextSwitchesProvider* ContextSwitchesProvider = TraceServices::ReadContextSwitchesProvider(*Session.Get());
-
-	if (ContextSwitchesProvider == nullptr)
-	{
-		return;
-	}
-
-	const uint64 NewThreadsSerial = ContextSwitchesProvider->GetThreadsSerial();
-	if (NewThreadsSerial == ThreadsSerial)
-	{
-		return;
-	}
-
-	ThreadsSerial = NewThreadsSerial;
-
 	//TODO: Create "CPU Thread" timing tracks also for threads without CPU timing events (i.e. only with context switch events).
-#endif
 
 	TSharedPtr<TimingProfiler::STimingView> TimingView = GetTimingView();
 	if (!TimingView.IsValid())
@@ -501,7 +475,7 @@ void FContextSwitchesSharedState::AddContextSwitchesChildTracks()
 void FContextSwitchesSharedState::RemoveContextSwitchesChildTracks()
 {
 	ThreadsSerial = 0;
-	
+
 	TSharedPtr<TimingProfiler::STimingView> TimingView = GetTimingView();
 	if (!TimingView.IsValid())
 	{
@@ -600,7 +574,7 @@ void FContextSwitchesSharedState::SetNonTargetProcessEventsVisible(bool bOnOff)
 				}
 			}
 		}
-		
+
 		bAreNonTargetProcessEventsVisible = bOnOff;
 	}
 }
@@ -682,7 +656,7 @@ void FContextSwitchesSharedState::Command_NavigateToCpuThreadEvent_Execute()
 ////////////////////////////////////////////////////////////////////////////////////////////////////
 
 bool FContextSwitchesSharedState::Command_NavigateToCpuThreadEvent_CanExecute() const
-{ 
+{
 	if (!AreContextSwitchesAvailable() || !AreContextSwitchesVisible() || !IsValidCpuCoreEventSelected())
 	{
 		return false;
@@ -833,15 +807,19 @@ void FContextSwitchesSharedState::GetThreadInfo(uint32 InSystemThreadId, uint32&
 	TSharedPtr<const TraceServices::IAnalysisSession> Session = FInsightsManager::Get()->GetSession();
 	if (Session.IsValid())
 	{
-		TraceServices::FAnalysisSessionReadScope SessionReadScope(*Session.Get());
-		const TraceServices::IContextSwitchesProvider* ContextSwitchesProvider = TraceServices::ReadContextSwitchesProvider(*Session.Get());
-		if (ContextSwitchesProvider)
+		uint32 ThreadId = 0;
+		bool bIsValidThreadId = false;
+		if (const TraceServices::IContextSwitchesProvider* ContextSwitchesProvider = TraceServices::ReadContextSwitchesProvider(*Session.Get()))
 		{
-			if (ContextSwitchesProvider->GetThreadId(InSystemThreadId, OutThreadId))
-			{
-				const TraceServices::IThreadProvider& ThreadProvider = TraceServices::ReadThreadProvider(*Session.Get());
-				OutThreadName = ThreadProvider.GetThreadName(OutThreadId);
-			}
+			TraceServices::FProviderReadScopeLock ScopedLock(*ContextSwitchesProvider);
+			bIsValidThreadId = ContextSwitchesProvider->GetThreadId(InSystemThreadId, ThreadId);
+		}
+		if (bIsValidThreadId)
+		{
+			const TraceServices::IThreadProvider& ThreadProvider = TraceServices::ReadThreadProvider(*Session.Get());
+			TraceServices::FAnalysisSessionReadScope SessionReadScope(*Session.Get());
+			OutThreadId = ThreadId;
+			OutThreadName = ThreadProvider.GetThreadName(ThreadId);
 		}
 	}
 }
diff --git a/Engine/Source/Developer/TraceInsights/Private/Insights/ContextSwitches/ViewModels/ContextSwitchesTimingTrack.cpp b/Engine/Source/Developer/TraceInsights/Private/Insights/ContextSwitches/ViewModels/ContextSwitchesTimingTrack.cpp
--- a/Engine/Source/Developer/TraceInsights/Private/Insights/ContextSwitches/ViewModels/ContextSwitchesTimingTrack.cpp
+++ b/Engine/Source/Developer/TraceInsights/Private/Insights/ContextSwitches/ViewModels/ContextSwitchesTimingTrack.cpp
@@ -5,6 +5,7 @@
 #include "Framework/MultiBox/MultiBoxBuilder.h"
 
 // TraceServices
+#include "Common/ProviderLock.h"
 #include "TraceServices/Model/ContextSwitches.h"
 
 // TraceInsightsCore
@@ -35,23 +36,21 @@ void FContextSwitchesTimingTrack::BuildDrawState(ITimingEventsTrackDrawStateBuil
 	{
 		return;
 	}
-	
-	TraceServices::FAnalysisSessionReadScope SessionReadScope(*Session.Get());
 
-	const TraceServices::IContextSwitchesProvider* ContextSwitchesProvider = TraceServices::ReadContextSwitchesProvider(*Session.Get());
-
-	if (ContextSwitchesProvider == nullptr)
+	if (const TraceServices::IContextSwitchesProvider* ContextSwitchesProvider = TraceServices::ReadContextSwitchesProvider(*Session.Get()))
 	{
-		return;
-	}
+		TraceServices::FProviderReadScopeLock ScopedLock(*ContextSwitchesProvider);
 
-	const FTimingTrackViewport& Viewport = Context.GetViewport();
+		const FTimingTrackViewport& Viewport = Context.GetViewport();
 
-	ContextSwitchesProvider->EnumerateContextSwitches(ThreadId, Viewport.GetStartTime(), Viewport.GetEndTime(), [&Builder](const TraceServices::FContextSwitch& ContextSwitch)
-		{
-			Builder.AddEvent(ContextSwitch.Start, ContextSwitch.End, 0, *FString::Printf(TEXT("Core %d"), ContextSwitch.CoreNumber), 0, 0);
-			return TraceServices::EContextSwitchEnumerationResult::Continue;
-		});
+		ContextSwitchesProvider->EnumerateContextSwitches(ThreadId, Viewport.GetStartTime(), Viewport.GetEndTime(),
+			[&Builder]
+			(const TraceServices::FContextSwitch& ContextSwitch)
+			{
+				Builder.AddEvent(ContextSwitch.Start, ContextSwitch.End, 0, *FString::Printf(TEXT("Core %d"), ContextSwitch.CoreNumber), 0, 0);
+				return TraceServices::EContextSwitchEnumerationResult::Continue;
+			});
+	}
 }
 
 ////////////////////////////////////////////////////////////////////////////////////////////////////
@@ -182,43 +181,45 @@ const TSharedPtr<const ITimingEvent> FContextSwitchesTimingTrack::GetEvent(doubl
 		return nullptr;
 	}
 
-	TraceServices::FAnalysisSessionReadScope SessionReadScope(*Session.Get());
-
-	const TraceServices::IContextSwitchesProvider* ContextSwitchesProvider = TraceServices::ReadContextSwitchesProvider(*Session.Get());
-	if (ContextSwitchesProvider == nullptr)
-	{
-		return nullptr;
-	}
-
 	TraceServices::FContextSwitch BestMatchContextSwitch;
 	double Delta = 2 * SecondsPerPixel;
 
-	ContextSwitchesProvider->EnumerateContextSwitches(ThreadId, InTime - Delta, InTime + Delta,
-		[InTime, &BestMatchContextSwitch, &Delta](const TraceServices::FContextSwitch& ContextSwitch)
-		{
-			if (ContextSwitch.Start <= InTime && ContextSwitch.End >= InTime)
-			{
-				BestMatchContextSwitch = ContextSwitch;
-				Delta = 0.0f;
-				return TraceServices::EContextSwitchEnumerationResult::Stop;
-			}
+	if (const TraceServices::IContextSwitchesProvider* ContextSwitchesProvider = TraceServices::ReadContextSwitchesProvider(*Session.Get()))
+	{
+		TraceServices::FProviderReadScopeLock ScopedLock(*ContextSwitchesProvider);
 
-			double DeltaLeft = InTime - ContextSwitch.End;
-			if (DeltaLeft >= 0 && DeltaLeft < Delta)
+		ContextSwitchesProvider->EnumerateContextSwitches(ThreadId, InTime - Delta, InTime + Delta,
+			[InTime, &BestMatchContextSwitch, &Delta]
+			(const TraceServices::FContextSwitch& ContextSwitch)
 			{
-				Delta = DeltaLeft;
-				BestMatchContextSwitch = ContextSwitch;
-			}
+				if (ContextSwitch.Start <= InTime && ContextSwitch.End >= InTime)
+				{
+					BestMatchContextSwitch = ContextSwitch;
+					Delta = 0.0f;
+					return TraceServices::EContextSwitchEnumerationResult::Stop;
+				}
 
-			double DeltaRight = ContextSwitch.Start - InTime;
-			if (DeltaRight >= 0 && DeltaRight < Delta)
-			{
-				Delta = DeltaRight;
-				BestMatchContextSwitch = ContextSwitch;
-			}
+				double DeltaLeft = InTime - ContextSwitch.End;
+				if (DeltaLeft >= 0 && DeltaLeft < Delta)
+				{
+					Delta = DeltaLeft;
+					BestMatchContextSwitch = ContextSwitch;
+				}
 
-			return TraceServices::EContextSwitchEnumerationResult::Continue;
-		});
+				double DeltaRight = ContextSwitch.Start - InTime;
+				if (DeltaRight >= 0 && DeltaRight < Delta)
+				{
+					Delta = DeltaRight;
+					BestMatchContextSwitch = ContextSwitch;
+				}
+
+				return TraceServices::EContextSwitchEnumerationResult::Continue;
+			});
+	}
+	else
+	{
+		return nullptr;
+	}
 
 	if (Delta < 2 * SecondsPerPixel)
 	{
@@ -282,45 +283,46 @@ const TSharedPtr<const ITimingEvent> FContextSwitchesTimingTrack::SearchEvent(co
 	TSharedPtr<const ITimingEvent> FoundEvent;
 
 	TTimingEventSearch<TraceServices::FContextSwitch>::Search(
+
 		InSearchParameters,
 
-		[this, &InSearchParameters](TTimingEventSearch<TraceServices::FContextSwitch>::FContext& InContext)
+		// Search Predicate
+		[this, &InSearchParameters]
+		(TTimingEventSearch<TraceServices::FContextSwitch>::FContext& InContext)
 		{
 			TSharedPtr<const TraceServices::IAnalysisSession> Session = FInsightsManager::Get()->GetSession();
 			if (Session.IsValid())
 			{
-				TraceServices::FAnalysisSessionReadScope SessionReadScope(*Session.Get());
-				if (Session.IsValid() && TraceServices::ReadContextSwitchesProvider(*Session.Get()))
+				if (const TraceServices::IContextSwitchesProvider* ContextSwitchesProvider = TraceServices::ReadContextSwitchesProvider(*Session.Get()))
 				{
-					const TraceServices::IContextSwitchesProvider* ContextSwitchesProvider = TraceServices::ReadContextSwitchesProvider(*Session.Get());
-
-					auto Callback = [&InContext](double EventStartTime, double EventEndTime, uint32 EventDepth, const TraceServices::FContextSwitch& Event)
-					{
-						InContext.Check(EventStartTime, EventEndTime, EventDepth, Event);
-						return InContext.ShouldContinueSearching() ? TraceServices::EContextSwitchEnumerationResult::Continue : TraceServices::EContextSwitchEnumerationResult::Stop;
-					};
-
+					TraceServices::FProviderReadScopeLock ScopedLock(*ContextSwitchesProvider);
 					ContextSwitchesProvider->EnumerateContextSwitches(ThreadId, InSearchParameters.StartTime, InSearchParameters.EndTime,
-						[ContextSwitchesProvider, Callback](const TraceServices::FContextSwitch& ContextSwitchEvent)
+						[&InContext]
+						(const TraceServices::FContextSwitch& ContextSwitchEvent)
 						{
-							return Callback(ContextSwitchEvent.Start, ContextSwitchEvent.End, 0, ContextSwitchEvent);
+							InContext.Check(ContextSwitchEvent.Start, ContextSwitchEvent.End, 0, ContextSwitchEvent);
+							return InContext.ShouldContinueSearching() ?
+								TraceServices::EContextSwitchEnumerationResult::Continue :
+								TraceServices::EContextSwitchEnumerationResult::Stop;
 						});
 				}
 			}
 		},
 
-		[](double EventStartTime, double EventEndTime, uint32 EventDepth, const TraceServices::FContextSwitch& Event)
-		{
-			return true;
-		},
+		// Filter Predicate
+		TTimingEventSearch<TraceServices::FContextSwitch>::NoFilter,
 
-		[&FoundEvent, this](double InFoundStartTime, double InFoundEndTime, uint32 InFoundDepth, const TraceServices::FContextSwitch& InEvent)
+		// Found Predicate
+		[this, &FoundEvent]
+		(double InFoundStartTime, double InFoundEndTime, uint32 InFoundDepth, const TraceServices::FContextSwitch& InEvent)
 		{
 			FoundEvent = MakeShared<FTimingEvent>(SharedThis(this), InFoundStartTime, InFoundEndTime, InFoundDepth, InEvent.CoreNumber);
 		},
 
+		// Payload Matched Predicate
 		TTimingEventSearch<TraceServices::FContextSwitch>::NoMatch
-	);
+
+	); // Search
 
 	return FoundEvent;
 }
diff --git a/Engine/Source/Developer/TraceInsights/Private/Insights/ContextSwitches/ViewModels/CpuCoreTimingTrack.cpp b/Engine/Source/Developer/TraceInsights/Private/Insights/ContextSwitches/ViewModels/CpuCoreTimingTrack.cpp
--- a/Engine/Source/Developer/TraceInsights/Private/Insights/ContextSwitches/ViewModels/CpuCoreTimingTrack.cpp
+++ b/Engine/Source/Developer/TraceInsights/Private/Insights/ContextSwitches/ViewModels/CpuCoreTimingTrack.cpp
@@ -5,6 +5,7 @@
 #include "Framework/MultiBox/MultiBoxBuilder.h"
 
 // TraceServices
+#include "Common/ProviderLock.h"
 #include "TraceServices/Model/ContextSwitches.h"
 #include "TraceServices/Model/Threads.h"
 
@@ -59,34 +60,33 @@ void FCpuCoreTimingTrack::BuildDrawState(ITimingEventsTrackDrawStateBuilder& Bui
 		return;
 	}
 
-	TraceServices::FAnalysisSessionReadScope SessionReadScope(*Session.Get());
-
-	const TraceServices::IContextSwitchesProvider* ContextSwitchesProvider = TraceServices::ReadContextSwitchesProvider(*Session.Get());
-	if (ContextSwitchesProvider == nullptr)
+	if (const TraceServices::IContextSwitchesProvider* ContextSwitchesProvider = TraceServices::ReadContextSwitchesProvider(*Session.Get()))
 	{
-		return;
-	}
+		TraceServices::FProviderReadScopeLock ScopedLock(*ContextSwitchesProvider);
 
-	const FTimingTrackViewport& Viewport = Context.GetViewport();
-	FTimingEventsTrackDrawStateBuilder NonTargetProcessEventsBuilder(*NonTargetProcessEventsDrawState, Viewport, Context.GetGeometry().Scale);
-	bool bShowNonTargetProcessEvents = SharedState.AreNonTargetProcessEventsVisible();
-	ContextSwitchesProvider->EnumerateCpuCoreEvents(CoreNumber, Viewport.GetStartTime(), Viewport.GetEndTime(), 
-		[&Builder, &NonTargetProcessEventsBuilder, ContextSwitchesProvider, bShowNonTargetProcessEvents, this](const TraceServices::FCpuCoreEvent& CpuCoreEvent)
-		{
-			uint32 ThreadId = 0;
-			if (ContextSwitchesProvider->GetThreadId(CpuCoreEvent.SystemThreadId, ThreadId))
-			{
-				this->AddCoreTimingEvent(Builder, CpuCoreEvent);
-			}
-			else if(bShowNonTargetProcessEvents)
+		const FTimingTrackViewport& Viewport = Context.GetViewport();
+		FTimingEventsTrackDrawStateBuilder NonTargetProcessEventsBuilder(*NonTargetProcessEventsDrawState, Viewport, Context.GetGeometry().Scale);
+		bool bShowNonTargetProcessEvents = SharedState.AreNonTargetProcessEventsVisible();
+
+		ContextSwitchesProvider->EnumerateCpuCoreEvents(CoreNumber, Viewport.GetStartTime(), Viewport.GetEndTime(),
+			[this, &Builder, &NonTargetProcessEventsBuilder, ContextSwitchesProvider, bShowNonTargetProcessEvents]
+			(const TraceServices::FCpuCoreEvent& CpuCoreEvent)
 			{
-				this->AddCoreTimingEvent(NonTargetProcessEventsBuilder, CpuCoreEvent);
-			}
+				uint32 ThreadId = 0;
+				if (ContextSwitchesProvider->GetThreadId(CpuCoreEvent.SystemThreadId, ThreadId))
+				{
+					this->AddCoreTimingEvent(Builder, CpuCoreEvent);
+				}
+				else if (bShowNonTargetProcessEvents)
+				{
+					this->AddCoreTimingEvent(NonTargetProcessEventsBuilder, CpuCoreEvent);
+				}
 
-			return TraceServices::EContextSwitchEnumerationResult::Continue;
-		});
+				return TraceServices::EContextSwitchEnumerationResult::Continue;
+			});
 
-	NonTargetProcessEventsMaxDepth = NonTargetProcessEventsBuilder.GetMaxDepth();
+		NonTargetProcessEventsMaxDepth = NonTargetProcessEventsBuilder.GetMaxDepth();
+	}
 }
 
 ////////////////////////////////////////////////////////////////////////////////////////////////////
@@ -106,39 +106,41 @@ void FCpuCoreTimingTrack::BuildFilteredDrawState(ITimingEventsTrackDrawStateBuil
 		}
 
 		TSharedPtr<const TraceServices::IAnalysisSession> Session = FInsightsManager::Get()->GetSession();
-		if (Session.IsValid() && TraceServices::ReadContextSwitchesProvider(*Session.Get()))
+		if (Session.IsValid())
 		{
-			TraceServices::FAnalysisSessionReadScope SessionReadScope(*Session.Get());
-
-			const TraceServices::IContextSwitchesProvider* ContextSwitchesProvider = TraceServices::ReadContextSwitchesProvider(*Session.Get());
+			if (const TraceServices::IContextSwitchesProvider* ContextSwitchesProvider = TraceServices::ReadContextSwitchesProvider(*Session.Get()))
+			{
+				TraceServices::FProviderReadScopeLock ScopedLock(*ContextSwitchesProvider);
 
-			const FTimingTrackViewport& Viewport = Context.GetViewport();
+				const FTimingTrackViewport& Viewport = Context.GetViewport();
 
-			if (bFilterOnlyByEventType)
-			{
-				bool bShowNonTargetProcessEvents = SharedState.AreNonTargetProcessEventsVisible();
-				ContextSwitchesProvider->EnumerateCpuCoreEvents(CoreNumber, Viewport.GetStartTime(), Viewport.GetEndTime(), 
-					[&Builder, ContextSwitchesProvider, FilterEventType, this, bShowNonTargetProcessEvents](const TraceServices::FCpuCoreEvent& CpuCoreEvent)
-					{
-						if (CpuCoreEvent.SystemThreadId != FilterEventType)
+				if (bFilterOnlyByEventType)
+				{
+					bool bShowNonTargetProcessEvents = SharedState.AreNonTargetProcessEventsVisible();
+					ContextSwitchesProvider->EnumerateCpuCoreEvents(CoreNumber, Viewport.GetStartTime(), Viewport.GetEndTime(),
+						[this, &Builder, ContextSwitchesProvider, FilterEventType, bShowNonTargetProcessEvents]
+						(const TraceServices::FCpuCoreEvent& CpuCoreEvent)
 						{
-							return TraceServices::EContextSwitchEnumerationResult::Continue;
-						}
+							if (CpuCoreEvent.SystemThreadId != FilterEventType)
+							{
+								return TraceServices::EContextSwitchEnumerationResult::Continue;
+							}
 
-						uint32 ThreadId;
-						if (!bShowNonTargetProcessEvents && !ContextSwitchesProvider->GetThreadId(CpuCoreEvent.SystemThreadId, ThreadId))
-						{
-							return TraceServices::EContextSwitchEnumerationResult::Continue;
-						}
+							uint32 ThreadId;
+							if (!bShowNonTargetProcessEvents && !ContextSwitchesProvider->GetThreadId(CpuCoreEvent.SystemThreadId, ThreadId))
+							{
+								return TraceServices::EContextSwitchEnumerationResult::Continue;
+							}
 
-						this->AddCoreTimingEvent(Builder, CpuCoreEvent);
+							this->AddCoreTimingEvent(Builder, CpuCoreEvent);
 
-						return TraceServices::EContextSwitchEnumerationResult::Continue;
-					});
-			}
-			else // generic filter
-			{
-				//TODO: if (EventFilterPtr->FilterEvent(TimingEvent))
+							return TraceServices::EContextSwitchEnumerationResult::Continue;
+						});
+				}
+				else // generic filter
+				{
+					//TODO: if (EventFilterPtr->FilterEvent(TimingEvent))
+				}
 			}
 		}
 	}
@@ -167,35 +169,42 @@ void FCpuCoreTimingTrack::BuildFilteredDrawState(ITimingEventsTrackDrawStateBuil
 		FilterContext.AddFilterData<int64>(static_cast<int32>(EFilterField::CoreEventName), 0);
 
 		TSharedPtr<const TraceServices::IAnalysisSession> Session = FInsightsManager::Get()->GetSession();
-		TraceServices::FAnalysisSessionReadScope SessionReadScope(*Session.Get());
-		const TraceServices::IContextSwitchesProvider* ContextSwitchesProvider = TraceServices::ReadContextSwitchesProvider(*Session.Get());
-		const FTimingTrackViewport& Viewport = Context.GetViewport();
+		if (Session.IsValid())
+		{
+			if (const TraceServices::IContextSwitchesProvider* ContextSwitchesProvider = TraceServices::ReadContextSwitchesProvider(*Session.Get()))
+			{
+				TraceServices::FProviderReadScopeLock ScopedLock(*ContextSwitchesProvider);
 
-		bool bShowNonTargetProcessEvents = SharedState.AreNonTargetProcessEventsVisible();
+				const FTimingTrackViewport& Viewport = Context.GetViewport();
 
-		ContextSwitchesProvider->EnumerateCpuCoreEvents(CoreNumber, Viewport.GetStartTime(), Viewport.GetEndTime(), 
-			[&Builder, ContextSwitchesProvider, &FilterContext, &FilterConfigurator, bShowNonTargetProcessEvents, this](const TraceServices::FCpuCoreEvent& CpuCoreEvent)
-			{
-				FilterContext.SetFilterData<double>(static_cast<int32>(EFilterField::StartTime), CpuCoreEvent.Start);
-				FilterContext.SetFilterData<double>(static_cast<int32>(EFilterField::EndTime), CpuCoreEvent.End);
-				FilterContext.SetFilterData<double>(static_cast<int32>(EFilterField::Duration), CpuCoreEvent.End - CpuCoreEvent.Start);
-				FilterContext.SetFilterData<int64>(static_cast<int32>(EFilterField::CoreEventName), CpuCoreEvent.SystemThreadId);
+				bool bShowNonTargetProcessEvents = SharedState.AreNonTargetProcessEventsVisible();
 
-				if (!FilterConfigurator->ApplyFilters(FilterContext))
-				{
-					return TraceServices::EContextSwitchEnumerationResult::Continue;
-				}
+				ContextSwitchesProvider->EnumerateCpuCoreEvents(CoreNumber, Viewport.GetStartTime(), Viewport.GetEndTime(),
+					[this, &Builder, ContextSwitchesProvider, &FilterContext, &FilterConfigurator, bShowNonTargetProcessEvents]
+					(const TraceServices::FCpuCoreEvent& CpuCoreEvent)
+					{
+						FilterContext.SetFilterData<double>(static_cast<int32>(EFilterField::StartTime), CpuCoreEvent.Start);
+						FilterContext.SetFilterData<double>(static_cast<int32>(EFilterField::EndTime), CpuCoreEvent.End);
+						FilterContext.SetFilterData<double>(static_cast<int32>(EFilterField::Duration), CpuCoreEvent.End - CpuCoreEvent.Start);
+						FilterContext.SetFilterData<int64>(static_cast<int32>(EFilterField::CoreEventName), CpuCoreEvent.SystemThreadId);
 
-				uint32 ThreadId;
-				if (!bShowNonTargetProcessEvents && !ContextSwitchesProvider->GetThreadId(CpuCoreEvent.SystemThreadId, ThreadId))
-				{
-					return TraceServices::EContextSwitchEnumerationResult::Continue;
-				}
+						if (!FilterConfigurator->ApplyFilters(FilterContext))
+						{
+							return TraceServices::EContextSwitchEnumerationResult::Continue;
+						}
+
+						uint32 ThreadId;
+						if (!bShowNonTargetProcessEvents && !ContextSwitchesProvider->GetThreadId(CpuCoreEvent.SystemThreadId, ThreadId))
+						{
+							return TraceServices::EContextSwitchEnumerationResult::Continue;
+						}
 
-				this->AddCoreTimingEvent(Builder, CpuCoreEvent);
+						this->AddCoreTimingEvent(Builder, CpuCoreEvent);
 
-				return TraceServices::EContextSwitchEnumerationResult::Continue;
-			});
+						return TraceServices::EContextSwitchEnumerationResult::Continue;
+					});
+			}
+		}
 	}
 }
 
@@ -269,55 +278,60 @@ const TSharedPtr<const ITimingEvent> FCpuCoreTimingTrack::GetEvent(double InTime
 		return nullptr;
 	}
 
-	TraceServices::FAnalysisSessionReadScope SessionReadScope(*Session.Get());
-
-	const TraceServices::IContextSwitchesProvider* ContextSwitchesProvider = TraceServices::ReadContextSwitchesProvider(*Session.Get());
-	if (ContextSwitchesProvider == nullptr)
-	{
-		return nullptr;
-	}
-
 	TraceServices::FCpuCoreEvent BestMatchEvent;
 	double Delta = 2 * SecondsPerPixel;
-	bool bShowNonTargetProcessEvents = SharedState.AreNonTargetProcessEventsVisible();
 
-	auto ShouldShowEvent = [bShowNonTargetProcessEvents, ContextSwitchesProvider](const TraceServices::FCpuCoreEvent& CpuCoreEvent)
+	if (const TraceServices::IContextSwitchesProvider* ContextSwitchesProvider = TraceServices::ReadContextSwitchesProvider(*Session.Get()))
 	{
-		uint32 ThreadId;
-		if (!bShowNonTargetProcessEvents)
-		{
-			return ContextSwitchesProvider->GetThreadId(CpuCoreEvent.SystemThreadId, ThreadId);
-		}
+		TraceServices::FProviderReadScopeLock ScopedLock(*ContextSwitchesProvider);
 
-		return true;
-	};
+		bool bShowNonTargetProcessEvents = SharedState.AreNonTargetProcessEventsVisible();
 
-	ContextSwitchesProvider->EnumerateCpuCoreEvents(CoreNumber, InTime - Delta, InTime + Delta,
-		[InTime, &BestMatchEvent, &Delta, ShouldShowEvent](const TraceServices::FCpuCoreEvent& CpuCoreEvent)
-		{
-			if (CpuCoreEvent.Start <= InTime && CpuCoreEvent.End >= InTime && ShouldShowEvent(CpuCoreEvent))
+		auto ShouldShowEvent =
+			[bShowNonTargetProcessEvents, ContextSwitchesProvider]
+			(const TraceServices::FCpuCoreEvent& CpuCoreEvent)
 			{
-				BestMatchEvent = CpuCoreEvent;
-				Delta = 0.0f;
-				return TraceServices::EContextSwitchEnumerationResult::Stop;
-			}
+				uint32 ThreadId;
+				if (!bShowNonTargetProcessEvents)
+				{
+					return ContextSwitchesProvider->GetThreadId(CpuCoreEvent.SystemThreadId, ThreadId);
+				}
 
-			double DeltaLeft = InTime - CpuCoreEvent.End;
-			if (DeltaLeft >= 0 && DeltaLeft < Delta && ShouldShowEvent(CpuCoreEvent))
-			{
-				Delta = DeltaLeft;
-				BestMatchEvent = CpuCoreEvent;
-			}
+				return true;
+			};
 
-			double DeltaRight = CpuCoreEvent.Start - InTime;
-			if (DeltaRight >= 0 && DeltaRight < Delta && ShouldShowEvent(CpuCoreEvent))
+		ContextSwitchesProvider->EnumerateCpuCoreEvents(CoreNumber, InTime - Delta, InTime + Delta,
+			[InTime, &BestMatchEvent, &Delta, ShouldShowEvent]
+			(const TraceServices::FCpuCoreEvent& CpuCoreEvent)
 			{
-				Delta = DeltaRight;
-				BestMatchEvent = CpuCoreEvent;
-			}
+				if (CpuCoreEvent.Start <= InTime && CpuCoreEvent.End >= InTime && ShouldShowEvent(CpuCoreEvent))
+				{
+					BestMatchEvent = CpuCoreEvent;
+					Delta = 0.0f;
+					return TraceServices::EContextSwitchEnumerationResult::Stop;
+				}
 
-			return TraceServices::EContextSwitchEnumerationResult::Continue;
-		});
+				double DeltaLeft = InTime - CpuCoreEvent.End;
+				if (DeltaLeft >= 0 && DeltaLeft < Delta && ShouldShowEvent(CpuCoreEvent))
+				{
+					Delta = DeltaLeft;
+					BestMatchEvent = CpuCoreEvent;
+				}
+
+				double DeltaRight = CpuCoreEvent.Start - InTime;
+				if (DeltaRight >= 0 && DeltaRight < Delta && ShouldShowEvent(CpuCoreEvent))
+				{
+					Delta = DeltaRight;
+					BestMatchEvent = CpuCoreEvent;
+				}
+
+				return TraceServices::EContextSwitchEnumerationResult::Continue;
+			});
+	}
+	else
+	{
+		return nullptr;
+	}
 
 	if (Delta < 2 * SecondsPerPixel)
 	{
@@ -344,45 +358,50 @@ const TSharedPtr<const ITimingEvent> FCpuCoreTimingTrack::SearchEvent(const FTim
 	FilterConfiguratorContext.AddFilterData<int64>(static_cast<int32>(EFilterField::CoreEventName), 0);
 
 	TTimingEventSearch<TraceServices::FCpuCoreEvent>::Search(
+
 		InSearchParameters,
 
-		[this, &InSearchParameters](TTimingEventSearch<TraceServices::FCpuCoreEvent>::FContext& InContext)
+		// Search Predicate
+		[this, &InSearchParameters]
+		(TTimingEventSearch<TraceServices::FCpuCoreEvent>::FContext& InContext)
 		{
 			TSharedPtr<const TraceServices::IAnalysisSession> Session = FInsightsManager::Get()->GetSession();
 			if (Session.IsValid())
 			{
-				TraceServices::FAnalysisSessionReadScope SessionReadScope(*Session.Get());
-				if (Session.IsValid() && TraceServices::ReadContextSwitchesProvider(*Session.Get()))
+				if (const TraceServices::IContextSwitchesProvider* ContextSwitchesProvider = TraceServices::ReadContextSwitchesProvider(*Session.Get()))
 				{
-					const TraceServices::IContextSwitchesProvider* ContextSwitchesProvider = TraceServices::ReadContextSwitchesProvider(*Session.Get());
-
-					auto Callback = [&InContext](double EventStartTime, double EventEndTime, uint32 EventDepth, const TraceServices::FCpuCoreEvent& Event)
-					{
-						InContext.Check(EventStartTime, EventEndTime, EventDepth, Event);
-						return InContext.ShouldContinueSearching() ? TraceServices::EContextSwitchEnumerationResult::Continue : TraceServices::EContextSwitchEnumerationResult::Stop;
-					};
-
+					TraceServices::FProviderReadScopeLock ScopedLock(*ContextSwitchesProvider);
 					if (InContext.GetParameters().SearchDirection == FTimingEventSearchParameters::ESearchDirection::Forward)
 					{
 						ContextSwitchesProvider->EnumerateCpuCoreEvents(CoreNumber, InSearchParameters.StartTime, InSearchParameters.EndTime,
-							[ContextSwitchesProvider, Callback](const TraceServices::FCpuCoreEvent& CpuCoreEvent)
+							[&InContext]
+							(const TraceServices::FCpuCoreEvent& CpuCoreEvent)
 							{
-								return Callback(CpuCoreEvent.Start, CpuCoreEvent.End, 0, CpuCoreEvent);
+								InContext.Check(CpuCoreEvent.Start, CpuCoreEvent.End, 0, CpuCoreEvent);
+								return InContext.ShouldContinueSearching() ?
+									TraceServices::EContextSwitchEnumerationResult::Continue :
+									TraceServices::EContextSwitchEnumerationResult::Stop;
 							});
 					}
 					else
 					{
 						ContextSwitchesProvider->EnumerateCpuCoreEventsBackwards(CoreNumber, InSearchParameters.EndTime, InSearchParameters.StartTime,
-							[ContextSwitchesProvider, Callback](const TraceServices::FCpuCoreEvent& CpuCoreEvent)
+							[&InContext]
+							(const TraceServices::FCpuCoreEvent& CpuCoreEvent)
 							{
-								return Callback(CpuCoreEvent.Start, CpuCoreEvent.End, 0, CpuCoreEvent);
+								InContext.Check(CpuCoreEvent.Start, CpuCoreEvent.End, 0, CpuCoreEvent);
+								return InContext.ShouldContinueSearching() ?
+									TraceServices::EContextSwitchEnumerationResult::Continue :
+									TraceServices::EContextSwitchEnumerationResult::Stop;
 							});
 					}
 				}
 			}
 		},
 
-		[&FilterConfiguratorContext, &InSearchParameters](double EventStartTime, double EventEndTime, uint32 EventDepth, const TraceServices::FCpuCoreEvent& Event)
+		// Filter Predicate
+		[&FilterConfiguratorContext, &InSearchParameters]
+		(double EventStartTime, double EventEndTime, uint32 EventDepth, const TraceServices::FCpuCoreEvent& Event)
 		{
 			if (!InSearchParameters.FilterExecutor.IsValid())
 			{
@@ -390,26 +409,30 @@ const TSharedPtr<const ITimingEvent> FCpuCoreTimingTrack::SearchEvent(const FTim
 			}
 
 			TSharedPtr<const TraceServices::IAnalysisSession> Session = FInsightsManager::Get()->GetSession();
-			if (Session.IsValid())
+			if (!Session.IsValid())
 			{
-				FilterConfiguratorContext.SetFilterData<double>(static_cast<int32>(EFilterField::StartTime), EventStartTime);
-				FilterConfiguratorContext.SetFilterData<double>(static_cast<int32>(EFilterField::EndTime), EventEndTime);
-				FilterConfiguratorContext.SetFilterData<double>(static_cast<int32>(EFilterField::Duration), EventEndTime - EventStartTime);
-				FilterConfiguratorContext.SetFilterData<int64>(static_cast<int32>(EFilterField::CoreEventName), Event.SystemThreadId);
-
-				return InSearchParameters.FilterExecutor->ApplyFilters(FilterConfiguratorContext);
+				return false;
 			}
 
-			return false;
+			FilterConfiguratorContext.SetFilterData<double>(static_cast<int32>(EFilterField::StartTime), EventStartTime);
+			FilterConfiguratorContext.SetFilterData<double>(static_cast<int32>(EFilterField::EndTime), EventEndTime);
+			FilterConfiguratorContext.SetFilterData<double>(static_cast<int32>(EFilterField::Duration), EventEndTime - EventStartTime);
+			FilterConfiguratorContext.SetFilterData<int64>(static_cast<int32>(EFilterField::CoreEventName), Event.SystemThreadId);
+
+			return InSearchParameters.FilterExecutor->ApplyFilters(FilterConfiguratorContext);
 		},
 
-		[&FoundEvent, this](double InFoundStartTime, double InFoundEndTime, uint32 InFoundDepth, const TraceServices::FCpuCoreEvent& InEvent)
+		// Found Predicate
+		[&FoundEvent, this]
+		(double InFoundStartTime, double InFoundEndTime, uint32 InFoundDepth, const TraceServices::FCpuCoreEvent& InEvent)
 		{
 			FoundEvent = MakeShared<FTimingEvent>(SharedThis(this), InFoundStartTime, InFoundEndTime, InFoundDepth, InEvent.SystemThreadId);
 		},
 
+		// Payload Matched Predicate
 		TTimingEventSearch<TraceServices::FCpuCoreEvent>::NoMatch
-		);
+
+	); // Search
 
 	return FoundEvent;
 }
@@ -504,16 +527,18 @@ FString FCpuCoreTimingTrack::GetThreadName(uint32 InSystemThreadId) const
 	TSharedPtr<const TraceServices::IAnalysisSession> Session = FInsightsManager::Get()->GetSession();
 	if (Session.IsValid())
 	{
-		TraceServices::FAnalysisSessionReadScope SessionReadScope(*Session.Get());
-		const TraceServices::IContextSwitchesProvider* ContextSwitchesProvider = TraceServices::ReadContextSwitchesProvider(*Session.Get());
-		if (ContextSwitchesProvider)
+		uint32 ThreadId = 0;
+		bool bIsValidThreadId = false;
+		if (const TraceServices::IContextSwitchesProvider* ContextSwitchesProvider = TraceServices::ReadContextSwitchesProvider(*Session.Get()))
 		{
-			uint32 ThreadId;
-			if (ContextSwitchesProvider->GetThreadId(InSystemThreadId, ThreadId))
-			{
-				const TraceServices::IThreadProvider& ThreadProvider = TraceServices::ReadThreadProvider(*Session.Get());
-				ThreadName = ThreadProvider.GetThreadName(ThreadId);
-			}
+			TraceServices::FProviderReadScopeLock ScopedLock(*ContextSwitchesProvider);
+			bIsValidThreadId = ContextSwitchesProvider->GetThreadId(InSystemThreadId, ThreadId);
+		}
+		if (bIsValidThreadId)
+		{
+			const TraceServices::IThreadProvider& ThreadProvider = TraceServices::ReadThreadProvider(*Session.Get());
+			TraceServices::FAnalysisSessionReadScope SessionReadScope(*Session.Get());
+			ThreadName = ThreadProvider.GetThreadName(ThreadId);
 		}
 	}
 
diff --git a/Engine/Source/Developer/TraceServices/Private/Analyzers/PlatformEventTraceAnalysis.cpp b/Engine/Source/Developer/TraceServices/Private/Analyzers/PlatformEventTraceAnalysis.cpp
--- a/Engine/Source/Developer/TraceServices/Private/Analyzers/PlatformEventTraceAnalysis.cpp
+++ b/Engine/Source/Developer/TraceServices/Private/Analyzers/PlatformEventTraceAnalysis.cpp
@@ -32,27 +32,31 @@ bool FPlatformEventTraceAnalyzer::OnEvent(uint16 RouteId, EStyle Style, const FO
 {
 	LLM_SCOPE_BYNAME(TEXT("Insights/FPlatformEventTraceAnalyzer"));
 
-	FAnalysisSessionEditScope _(Session);
-
-	const auto& EventData = Context.EventData;
 	switch (RouteId)
 	{
 
 	case RouteId_ContextSwitch:
 	{
+		const auto& EventData = Context.EventData;
 		double Start = Context.EventTime.AsSeconds(EventData.GetValue<uint64>("StartTime"));
 		double End = Context.EventTime.AsSeconds(EventData.GetValue<uint64>("EndTime"));
 		uint32 ThreadId = EventData.GetValue<uint32>("ThreadId");
 		uint32 CoreNumber = EventData.GetValue<uint8>("CoreNumber");
-		ContextSwitchesProvider.Add(ThreadId, Start, End, CoreNumber);
-
-		Session.UpdateDurationSeconds(End);
-
+		{
+			TraceServices::FProviderEditScopeLock ScopedLock(ContextSwitchesProvider);
+			ContextSwitchesProvider.AddTimingEvent(ThreadId, Start, End, CoreNumber);
+		}
+		{
+			FAnalysisSessionEditScope _(Session);
+			Session.UpdateDurationSeconds(End);
+		}
 		break;
 	}
 
 	case RouteId_StackSample:
 	{
+		FAnalysisSessionEditScope _(Session);
+		const auto& EventData = Context.EventData;
 		double Time = Context.EventTime.AsSeconds(EventData.GetValue<uint64>("Time"));
 		uint32 ThreadId = EventData.GetValue<uint32>("ThreadId");
 		const TArrayReader<uint64>& Addresses = EventData.GetArray<uint64>("Addresses");
@@ -65,17 +69,19 @@ bool FPlatformEventTraceAnalyzer::OnEvent(uint16 RouteId, EStyle Style, const FO
 
 	case RouteId_ThreadName:
 	{
+		const auto& EventData = Context.EventData;
 		uint32 ThreadId = EventData.GetValue<uint32>("ThreadId");
 		uint32 ProcessId = EventData.GetValue<uint32>("ProcessId");
 		FStringView Name;
 		if (EventData.GetString("Name", Name))
 		{
+			TraceServices::FProviderEditScopeLock ScopedLock(ContextSwitchesProvider);
 			ContextSwitchesProvider.AddThreadName(ThreadId, ProcessId, Name);
 		}
 		break;
 	}
 
-	}
+	} // switch (RouteId)
 
 	return true;
 }
@@ -84,7 +90,7 @@ void FPlatformEventTraceAnalyzer::OnThreadInfo(const FThreadInfo& ThreadInfo)
 {
 	LLM_SCOPE_BYNAME(TEXT("Insights/FPlatformEventTraceAnalyzer"));
 
-	FAnalysisSessionEditScope _(Session);
+	TraceServices::FProviderEditScopeLock ScopedLock(ContextSwitchesProvider);
 	ContextSwitchesProvider.AddThreadInfo(ThreadInfo.GetId(), ThreadInfo.GetSystemId());
 }
 
diff --git a/Engine/Source/Developer/TraceServices/Private/Model/ContextSwitches.cpp b/Engine/Source/Developer/TraceServices/Private/Model/ContextSwitches.cpp
--- a/Engine/Source/Developer/TraceServices/Private/Model/ContextSwitches.cpp
+++ b/Engine/Source/Developer/TraceServices/Private/Model/ContextSwitches.cpp
@@ -4,11 +4,15 @@
 #include "Model/ContextSwitchesPrivate.h"
 
 #include "Algo/Sort.h"
+
+// TraceServices
 #include "AnalysisServicePrivate.h"
 
 namespace TraceServices
 {
 
+thread_local FProviderLock::FThreadLocalState GContextSwitchesProviderLockState;
+
 FContextSwitchesProvider::FContextSwitchesProvider(IAnalysisSession& InSession)
 	: Session(InSession)
 	, NumCpuCores(0)
@@ -35,14 +39,14 @@ FContextSwitchesProvider::~FContextSwitchesProvider()
 
 bool FContextSwitchesProvider::HasData() const
 {
-	Session.ReadAccessCheck();
+	ReadAccessCheck();
 
 	return Threads.Num() > 0;
 }
 
 bool FContextSwitchesProvider::GetSystemThreadId(uint32 ThreadId, uint32& OutSystemThreadId) const
 {
-	Session.ReadAccessCheck();
+	ReadAccessCheck();
 
 	const uint32* SystemThreadIdPtr = TraceToSystemThreadIdMap.Find(ThreadId);
 	if (SystemThreadIdPtr)
@@ -55,7 +59,7 @@ bool FContextSwitchesProvider::GetSystemThreadId(uint32 ThreadId, uint32& OutSys
 
 bool FContextSwitchesProvider::GetThreadId(uint32 SystemThreadId, uint32& OutThreadId) const
 {
-	Session.ReadAccessCheck();
+	ReadAccessCheck();
 
 	const uint32* ThreadIdPtr = SystemToTraceThreadIdMap.Find(SystemThreadId);
 	if (ThreadIdPtr)
@@ -68,7 +72,7 @@ bool FContextSwitchesProvider::GetThreadId(uint32 SystemThreadId, uint32& OutThr
 
 bool FContextSwitchesProvider::GetSystemThreadId(uint32 CoreNumber, double Time, uint32& OutSystemThreadId) const
 {
-	Session.ReadAccessCheck();
+	ReadAccessCheck();
 
 	if (CoreNumber >= (uint32)CpuCores.Num())
 	{
@@ -110,7 +114,7 @@ bool FContextSwitchesProvider::GetSystemThreadId(uint32 CoreNumber, double Time,
 
 bool FContextSwitchesProvider::GetThreadId(uint32 CoreNumber, double Time, uint32& OutThreadId) const
 {
-	Session.ReadAccessCheck();
+	ReadAccessCheck();
 
 	uint32 SystemThreadId;
 	if (GetSystemThreadId(CoreNumber, Time, SystemThreadId))
@@ -123,7 +127,7 @@ bool FContextSwitchesProvider::GetThreadId(uint32 CoreNumber, double Time, uint3
 
 bool FContextSwitchesProvider::GetCoreNumber(uint32 ThreadId, double Time, uint32& OutCoreNumber) const
 {
-	Session.ReadAccessCheck();
+	ReadAccessCheck();
 
 	const TPagedArray<FContextSwitch>* ContextSwitches = GetContextSwitches(ThreadId);
 	if (!ContextSwitches)
@@ -160,7 +164,7 @@ bool FContextSwitchesProvider::GetCoreNumber(uint32 ThreadId, double Time, uint3
 
 void FContextSwitchesProvider::EnumerateCpuCores(CpuCoreCallback Callback) const
 {
-	Session.ReadAccessCheck();
+	ReadAccessCheck();
 
 	const uint32 CoreCount = (uint32)CpuCores.Num();
 	for (uint32 CoreNumber = 0; CoreNumber < CoreCount; ++CoreNumber)
@@ -176,7 +180,7 @@ void FContextSwitchesProvider::EnumerateCpuCores(CpuCoreCallback Callback) const
 
 void FContextSwitchesProvider::EnumerateCpuCoreEvents(uint32 CoreNumber, double StartTime, double EndTime, CpuCoreEventCallback Callback) const
 {
-	Session.ReadAccessCheck();
+	ReadAccessCheck();
 
 	if (CoreNumber >= (uint32)CpuCores.Num())
 	{
@@ -217,7 +221,7 @@ void FContextSwitchesProvider::EnumerateCpuCoreEvents(uint32 CoreNumber, double
 
 void FContextSwitchesProvider::EnumerateCpuCoreEventsBackwards(uint32 CoreNumber, double EndTime, double StartTime, CpuCoreEventCallback Callback) const
 {
-	Session.ReadAccessCheck();
+	ReadAccessCheck();
 
 	if (CoreNumber >= (uint32)CpuCores.Num())
 	{
@@ -260,7 +264,7 @@ void FContextSwitchesProvider::EnumerateCpuCoreEventsBackwards(uint32 CoreNumber
 
 const TPagedArray<FContextSwitch>* FContextSwitchesProvider::GetContextSwitches(uint32 ThreadId) const
 {
-	Session.ReadAccessCheck();
+	ReadAccessCheck();
 
 	if (TraceToSystemThreadIdMap.Contains(ThreadId))
 	{
@@ -273,7 +277,7 @@ const TPagedArray<FContextSwitch>* FContextSwitchesProvider::GetContextSwitches(
 
 void FContextSwitchesProvider::EnumerateContextSwitches(uint32 ThreadId, double StartTime, double EndTime, ContextSwitchCallback Callback) const
 {
-	Session.ReadAccessCheck();
+	ReadAccessCheck();
 
 	const TPagedArray<FContextSwitch>* ContextSwitches = GetContextSwitches(ThreadId);
 	if (ContextSwitches == nullptr)
@@ -307,9 +311,9 @@ void FContextSwitchesProvider::EnumerateContextSwitches(uint32 ThreadId, double
 	}
 }
 
-void FContextSwitchesProvider::Add(uint32 SystemThreadId, double Start, double End, uint32 CoreNumber)
+void FContextSwitchesProvider::AddTimingEvent(uint32 SystemThreadId, double Start, double End, uint32 CoreNumber)
 {
-	Session.WriteAccessCheck();
+	EditAccessCheck();
 
 	//////////////////////////////////////////////////
 	// Context Switch events
@@ -360,7 +364,7 @@ void FContextSwitchesProvider::Add(uint32 SystemThreadId, double Start, double E
 
 void FContextSwitchesProvider::AddThreadInfo(uint32 ThreadId, uint32 SystemThreadId)
 {
-	Session.WriteAccessCheck();
+	EditAccessCheck();
 
 	if (SystemThreadId == 0)
 	{
@@ -386,7 +390,7 @@ void FContextSwitchesProvider::AddThreadInfo(uint32 ThreadId, uint32 SystemThrea
 	SystemToTraceThreadIdMap.Add(SystemThreadId, ThreadId);
 }
 
-void FContextSwitchesProvider::AddThreadName(uint32 SystemTreadId, uint32 SystemProcessId, FStringView Name)
+void FContextSwitchesProvider::AddThreadName(uint32 SystemThreadId, uint32 SystemProcessId, FStringView Name)
 {
 	//TODO
 }
diff --git a/Engine/Source/Developer/TraceServices/Private/Model/ContextSwitchesPrivate.h b/Engine/Source/Developer/TraceServices/Private/Model/ContextSwitchesPrivate.h
--- a/Engine/Source/Developer/TraceServices/Private/Model/ContextSwitchesPrivate.h
+++ b/Engine/Source/Developer/TraceServices/Private/Model/ContextSwitchesPrivate.h
@@ -2,44 +2,69 @@
 
 #pragma once
 
-#include "Common/PagedArray.h" // TraceServices
+#include "CoreTypes.h"
 #include "Containers/Map.h"
+
+// TraceServices
+#include "Common/PagedArray.h"
+#include "Common/ProviderLock.h"
 #include "TraceServices/Model/AnalysisSession.h"
 #include "TraceServices/Model/ContextSwitches.h"
 
 namespace TraceServices
 {
 
+extern thread_local FProviderLock::FThreadLocalState GContextSwitchesProviderLockState;
+
 class FContextSwitchesProvider
 	: public IContextSwitchesProvider
+	, public IEditableProvider
 {
 public:
 	explicit FContextSwitchesProvider(IAnalysisSession& Session);
 	virtual ~FContextSwitchesProvider();
 
+	//////////////////////////////////////////////////
+	// Read operations
+
+	virtual void BeginRead() const override { Lock.BeginRead(GContextSwitchesProviderLockState); }
+	virtual void EndRead() const override { Lock.EndRead(GContextSwitchesProviderLockState); }
+	virtual void ReadAccessCheck() const override { Lock.ReadAccessCheck(GContextSwitchesProviderLockState); }
+
 	virtual bool HasData() const override;
-	bool GetSystemThreadId(uint32 ThreadId, uint32& OutSystemThreadId) const override;
-	bool GetThreadId(uint32 SystemThreadId, uint32& OutThreadId) const override;
-	bool GetSystemThreadId(uint32 CoreNumber, double Time, uint32& OutSystemThreadId) const override;
-	bool GetThreadId(uint32 CoreNumber, double Time, uint32& OutThreadId) const override;
-	bool GetCoreNumber(uint32 ThreadId, double Time, uint32& OutCoreNumber) const override;
-	uint64 GetThreadsSerial() const override { return uint64(Threads.Num()); }
-	uint64 GetCpuCoresSerial() const override { return uint64(NumCpuCores); }
-	uint32 GetNumCpuCoresWithEvents() const { return NumCpuCores; }
-	uint32 GetExclusiveMaxCpuCoreNumber() const { return CpuCores.Num(); }
-	void EnumerateCpuCores(CpuCoreCallback Callback) const override;
-	void EnumerateContextSwitches(uint32 ThreadId, double StartTime, double EndTime, ContextSwitchCallback Callback) const override;
-	void EnumerateCpuCoreEvents(uint32 CoreNumber, double StartTime, double EndTime, CpuCoreEventCallback Callback) const override;
-	void EnumerateCpuCoreEventsBackwards(uint32 CoreNumber, double EndTime, double StartTime, CpuCoreEventCallback Callback) const override;
+	virtual bool GetSystemThreadId(uint32 ThreadId, uint32& OutSystemThreadId) const override;
+	virtual bool GetThreadId(uint32 SystemThreadId, uint32& OutThreadId) const override;
+	virtual bool GetSystemThreadId(uint32 CoreNumber, double Time, uint32& OutSystemThreadId) const override;
+	virtual bool GetThreadId(uint32 CoreNumber, double Time, uint32& OutThreadId) const override;
+	virtual bool GetCoreNumber(uint32 ThreadId, double Time, uint32& OutCoreNumber) const override;
+	virtual uint64 GetThreadsSerial() const override { ReadAccessCheck(); return uint64(Threads.Num()); }
+	virtual uint64 GetCpuCoresSerial() const override { ReadAccessCheck(); return uint64(NumCpuCores); }
+	virtual void EnumerateCpuCores(CpuCoreCallback Callback) const override;
+	virtual void EnumerateContextSwitches(uint32 ThreadId, double StartTime, double EndTime, ContextSwitchCallback Callback) const override;
+	virtual void EnumerateCpuCoreEvents(uint32 CoreNumber, double StartTime, double EndTime, CpuCoreEventCallback Callback) const override;
+	virtual void EnumerateCpuCoreEventsBackwards(uint32 CoreNumber, double EndTime, double StartTime, CpuCoreEventCallback Callback) const override;
 
-	void Add(uint32 SystemThreadId, double Start, double End, uint32 CoreNumber);
+	//////////////////////////////////////////////////
+	// Edit operations
+
+	virtual void BeginEdit() const override { Lock.BeginWrite(GContextSwitchesProviderLockState); }
+	virtual void EndEdit() const override { Lock.EndWrite(GContextSwitchesProviderLockState); }
+	virtual void EditAccessCheck() const override { Lock.WriteAccessCheck(GContextSwitchesProviderLockState); }
+
+	void AddTimingEvent(uint32 SystemThreadId, double StartTime, double EndTime, uint32 CoreNumber);
 	void AddThreadInfo(uint32 ThreadId, uint32 SystemThreadId);
+	void AddThreadName(uint32 SystemThreadId, uint32 SystemProcessId, FStringView Name);
 
-	void AddThreadName(uint32 SystemTreadId, uint32 SystemProcessId, FStringView Name);
+	//////////////////////////////////////////////////
 
 private:
+	uint32 GetNumCpuCoresWithEvents() const { return NumCpuCores; }
+	uint32 GetExclusiveMaxCpuCoreNumber() const { return CpuCores.Num(); }
 	const TPagedArray<FContextSwitch>* GetContextSwitches(uint32 ThreadId) const;
 
+private:
+	mutable FProviderLock Lock;
+
 	IAnalysisSession& Session;
 
 	// (Trace) Thread Id -> System Thread Id
