[Insights] Code refactoring re event search in the timing tracks (Loading, FileActivity and Frame timing tracks).

#rnx

diff --git a/Engine/Source/Developer/TraceInsights/Private/Insights/LoadingProfiler/Tracks/LoadingTimingTrack.cpp b/Engine/Source/Developer/TraceInsights/Private/Insights/LoadingProfiler/Tracks/LoadingTimingTrack.cpp
--- a/Engine/Source/Developer/TraceInsights/Private/Insights/LoadingProfiler/Tracks/LoadingTimingTrack.cpp
+++ b/Engine/Source/Developer/TraceInsights/Private/Insights/LoadingProfiler/Tracks/LoadingTimingTrack.cpp
@@ -205,33 +205,39 @@ bool FLoadingTimingTrack::FindLoadTimeProfilerCpuEvent(const FTimingEventSearchP
 	FilterConfiguratorContext.AddFilterData<FString>(static_cast<int32>(EFilterField::TrackName), this->GetName());
 
 	return TTimingEventSearch<TraceServices::FLoadTimeProfilerCpuEvent>::Search(
+
 		InParameters,
 
-		[this](TTimingEventSearch<TraceServices::FLoadTimeProfilerCpuEvent>::FContext& InContext)
+		// Search Predicate
+		[this]
+		(TTimingEventSearch<TraceServices::FLoadTimeProfilerCpuEvent>::FContext& InContext)
 		{
 			TSharedPtr<const TraceServices::IAnalysisSession> Session = FInsightsManager::Get()->GetSession();
 			if (Session.IsValid())
 			{
-				TraceServices::FAnalysisSessionReadScope SessionReadScope(*Session.Get());
-
 				if (TraceServices::ReadLoadTimeProfilerProvider(*Session.Get()))
 				{
 					const TraceServices::ILoadTimeProfilerProvider& LoadTimeProfilerProvider = *TraceServices::ReadLoadTimeProfilerProvider(*Session.Get());
+					TraceServices::FAnalysisSessionReadScope SessionReadScope(*Session.Get());
 
 					LoadTimeProfilerProvider.ReadTimeline(TimelineIndex, [&InContext](const TraceServices::ILoadTimeProfilerProvider::CpuTimeline& Timeline)
 					{
-						auto Callback = [&InContext](double EventStartTime, double EventEndTime, uint32 EventDepth, const TraceServices::FLoadTimeProfilerCpuEvent& Event)
-						{
-							if (Event.Package)
+						auto Callback =
+							[&InContext]
+							(double EventStartTime, double EventEndTime, uint32 EventDepth, const TraceServices::FLoadTimeProfilerCpuEvent& Event)
 							{
-								InContext.Check(EventStartTime, EventEndTime, EventDepth, Event);
-								return InContext.ShouldContinueSearching() ? TraceServices::EEventEnumerate::Continue : TraceServices::EEventEnumerate::Stop;
-							}
-							else
-							{
-								return TraceServices::EEventEnumerate::Continue;
-							}
-						};
+								if (Event.Package)
+								{
+									InContext.Check(EventStartTime, EventEndTime, EventDepth, Event);
+									return InContext.ShouldContinueSearching() ?
+										TraceServices::EEventEnumerate::Continue :
+										TraceServices::EEventEnumerate::Stop;
+								}
+								else
+								{
+									return TraceServices::EEventEnumerate::Continue;
+								}
+							};
 
 						if (InContext.GetParameters().SearchDirection == FTimingEventSearchParameters::ESearchDirection::Forward)
 						{
@@ -246,7 +252,9 @@ bool FLoadingTimingTrack::FindLoadTimeProfilerCpuEvent(const FTimingEventSearchP
 			}
 		},
 
-		[&FilterConfiguratorContext, &InParameters](double EventStartTime, double EventEndTime, uint32 EventDepth, const TraceServices::FLoadTimeProfilerCpuEvent& Event)
+		// Filter Predicate
+		[&FilterConfiguratorContext, &InParameters]
+		(double EventStartTime, double EventEndTime, uint32 EventDepth, const TraceServices::FLoadTimeProfilerCpuEvent& Event)
 		{
 			if (!InParameters.FilterExecutor.IsValid())
 			{
@@ -271,12 +279,13 @@ bool FLoadingTimingTrack::FindLoadTimeProfilerCpuEvent(const FTimingEventSearchP
 			return false;
 		},
 
-		[&InFoundPredicate](double InFoundStartTime, double InFoundEndTime, uint32 InFoundDepth, const TraceServices::FLoadTimeProfilerCpuEvent& InEvent)
-		{
-			InFoundPredicate(InFoundStartTime, InFoundEndTime, InFoundDepth, InEvent);
-		},
+		// Found Predicate
+		InFoundPredicate,
+
+		// Payload Matched Predicate
+		TTimingEventSearch<TraceServices::FLoadTimeProfilerCpuEvent>::NoMatch
 
-		TTimingEventSearch<TraceServices::FLoadTimeProfilerCpuEvent>::NoMatch);
+	); // Search
 }
 
 ////////////////////////////////////////////////////////////////////////////////////////////////////
diff --git a/Engine/Source/Developer/TraceInsights/Private/Insights/TimingProfiler/Tracks/FileActivityTimingTrack.cpp b/Engine/Source/Developer/TraceInsights/Private/Insights/TimingProfiler/Tracks/FileActivityTimingTrack.cpp
--- a/Engine/Source/Developer/TraceInsights/Private/Insights/TimingProfiler/Tracks/FileActivityTimingTrack.cpp
+++ b/Engine/Source/Developer/TraceInsights/Private/Insights/TimingProfiler/Tracks/FileActivityTimingTrack.cpp
@@ -167,10 +167,12 @@ void FFileActivityTimingTrack::InitTooltip(FTooltipDrawState& InOutTooltip, cons
 bool FFileActivityTimingTrack::FindIoTimingEvent(const FTimingEventSearchParameters& InParameters, TFunctionRef<void(double, double, uint32, const FIoTimingEvent&)> InFoundPredicate) const
 {
 	return TTimingEventSearch<FIoTimingEvent>::Search(
+
 		InParameters,
 
-		// Search...
-		[this](TTimingEventSearch<FIoTimingEvent>::FContext& InContext)
+		// Search Predicate
+		[this]
+		(TTimingEventSearch<FIoTimingEvent>::FContext& InContext)
 		{
 			const TArray<FIoTimingEvent>& Events = SharedState.GetAllEvents();
 
@@ -238,11 +240,14 @@ bool FFileActivityTimingTrack::FindIoTimingEvent(const FTimingEventSearchParamet
 			}
 		},
 
-		// Found!
-		[&InFoundPredicate](double InFoundStartTime, double InFoundEndTime, uint32 InFoundDepth, const FIoTimingEvent& InEvent)
-		{
-			InFoundPredicate(InFoundStartTime, InFoundEndTime, InFoundDepth, InEvent);
-		});
+		// No Filter Predicate
+
+		// Found Predicate
+		InFoundPredicate
+
+		// No Payload Matched Predicate
+
+	); // Search
 }
 
 ////////////////////////////////////////////////////////////////////////////////////////////////////
diff --git a/Engine/Source/Developer/TraceInsights/Private/Insights/TimingProfiler/ViewModels/FrameTimingTrack.cpp b/Engine/Source/Developer/TraceInsights/Private/Insights/TimingProfiler/ViewModels/FrameTimingTrack.cpp
--- a/Engine/Source/Developer/TraceInsights/Private/Insights/TimingProfiler/ViewModels/FrameTimingTrack.cpp
+++ b/Engine/Source/Developer/TraceInsights/Private/Insights/TimingProfiler/ViewModels/FrameTimingTrack.cpp
@@ -508,9 +508,12 @@ bool FFrameTimingTrack::FindFrame(const FTimingEvent& InTimingEvent, TFunctionRe
 bool FFrameTimingTrack::FindFrame(const FTimingEventSearchParameters& InParameters, TFunctionRef<void(double, double, uint32, const TraceServices::FFrame&)> InFoundPredicate) const
 {
 	return TTimingEventSearch<TraceServices::FFrame>::Search(
+
 		InParameters,
 
-		[this](TTimingEventSearch<TraceServices::FFrame>::FContext& InContext)
+		// Search Predicate
+		[this]
+		(TTimingEventSearch<TraceServices::FFrame>::FContext& InContext)
 		{
 			TSharedPtr<const TraceServices::IAnalysisSession> Session = FInsightsManager::Get()->GetSession();
 			if (Session.IsValid())
@@ -538,12 +541,17 @@ bool FFrameTimingTrack::FindFrame(const FTimingEventSearchParameters& InParamete
 			}
 		},
 
-		[&InFoundPredicate](double InFoundStartTime, double InFoundEndTime, uint32 InFoundDepth, const TraceServices::FFrame& InEvent)
-		{
-			InFoundPredicate(InFoundStartTime, InFoundEndTime, InFoundDepth, InEvent);
-		},
+		// No Filter Predicate
+
+		// Found Predicate
+		InFoundPredicate,
+
+		// No Payload Matched Predicate
+
+		// Cache
+		SearchCache
 
-		SearchCache);
+	); // Search
 }
 
 ////////////////////////////////////////////////////////////////////////////////////////////////////
