[Insights] Added initial UI for the new Stack Sampling profiler.

diff --git a/Engine/Source/Developer/TraceInsights/Private/Insights/InsightsSettings.cpp b/Engine/Source/Developer/TraceInsights/Private/Insights/InsightsSettings.cpp
--- a/Engine/Source/Developer/TraceInsights/Private/Insights/InsightsSettings.cpp
+++ b/Engine/Source/Developer/TraceInsights/Private/Insights/InsightsSettings.cpp
@@ -222,8 +222,9 @@ void FInsightsSettings::LoadFromConfig()
 
 	GConfig->GetInt(TEXT("Insights.TimingProfiler.TimersView"), TEXT("GroupingMode"), TimersViewGroupingMode, SettingsIni);
 	GConfig->GetBool(TEXT("Insights.TimingProfiler.TimersView"), TEXT("ShowGpuTimers"), bTimersViewShowGpuTimers, SettingsIni);
-	GConfig->GetBool(TEXT("Insights.TimingProfiler.TimersView"), TEXT("ShowVerseTimers"), bTimersViewShowVerseTimers, SettingsIni);
+	GConfig->GetBool(TEXT("Insights.TimingProfiler.TimersView"), TEXT("ShowVerseSamplingTimers"), bTimersViewShowVerseSamplingTimers, SettingsIni);
 	GConfig->GetBool(TEXT("Insights.TimingProfiler.TimersView"), TEXT("ShowCpuTimers"), bTimersViewShowCpuTimers, SettingsIni);
+	GConfig->GetBool(TEXT("Insights.TimingProfiler.TimersView"), TEXT("ShowCpuSamplingTimers"), bTimersViewShowCpuSamplingTimers, SettingsIni);
 	GConfig->GetBool(TEXT("Insights.TimingProfiler.TimersView"), TEXT("ShowZeroCountTimers"), bTimersViewShowZeroCountTimers, SettingsIni);
 
 	//////////////////////////////////////////////////
@@ -356,8 +357,9 @@ void FInsightsSettings::SaveToConfig()
 
 	GConfig->SetInt(TEXT("Insights.TimingProfiler.TimersView"), TEXT("GroupingMode"), TimersViewGroupingMode, SettingsIni);
 	GConfig->SetBool(TEXT("Insights.TimingProfiler.TimersView"), TEXT("ShowGpuTimers"), bTimersViewShowGpuTimers, SettingsIni);
-	GConfig->SetBool(TEXT("Insights.TimingProfiler.TimersView"), TEXT("ShowVerseTimers"), bTimersViewShowVerseTimers, SettingsIni);
+	GConfig->SetBool(TEXT("Insights.TimingProfiler.TimersView"), TEXT("ShowVerseSamplingTimers"), bTimersViewShowVerseSamplingTimers, SettingsIni);
 	GConfig->SetBool(TEXT("Insights.TimingProfiler.TimersView"), TEXT("ShowCpuTimers"), bTimersViewShowCpuTimers, SettingsIni);
+	GConfig->SetBool(TEXT("Insights.TimingProfiler.TimersView"), TEXT("ShowCpuSamplingTimers"), bTimersViewShowCpuSamplingTimers, SettingsIni);
 	GConfig->SetBool(TEXT("Insights.TimingProfiler.TimersView"), TEXT("ShowZeroCountTimers"), bTimersViewShowZeroCountTimers, SettingsIni);
 
 	//////////////////////////////////////////////////
diff --git a/Engine/Source/Developer/TraceInsights/Private/Insights/InsightsSettings.h b/Engine/Source/Developer/TraceInsights/Private/Insights/InsightsSettings.h
--- a/Engine/Source/Developer/TraceInsights/Private/Insights/InsightsSettings.h
+++ b/Engine/Source/Developer/TraceInsights/Private/Insights/InsightsSettings.h
@@ -460,17 +460,21 @@ public:
 	void SetTimersViewGroupingMode(int32 InValue) { TimersViewGroupingMode = InValue; }
 	void SetAndSaveTimersViewGroupingMode(int32 InValue) { SET_AND_SAVE(TimersViewGroupingMode, InValue); }
 
-	bool GetTimersViewShowGpuEvents() const { return bTimersViewShowGpuTimers; }
-	void SetTimersViewShowGpuEvents(bool InValue) { bTimersViewShowGpuTimers = InValue; }
-	void SetAndSaveTimersViewShowGpuEvents(bool InValue) { SET_AND_SAVE(bTimersViewShowGpuTimers, InValue); }
+	bool GetTimersViewShowGpuTimers() const { return bTimersViewShowGpuTimers; }
+	void SetTimersViewShowGpuTimers(bool InValue) { bTimersViewShowGpuTimers = InValue; }
+	void SetAndSaveTimersViewShowGpuTimers(bool InValue) { SET_AND_SAVE(bTimersViewShowGpuTimers, InValue); }
 
-	bool GetTimersViewShowVerseEvents() const { return bTimersViewShowVerseTimers; }
-	void SetTimersViewShowVerseEvents(bool InValue) { bTimersViewShowVerseTimers = InValue; }
-	void SetAndSaveTimersViewShowVerseEvents(bool InValue) { SET_AND_SAVE(bTimersViewShowVerseTimers, InValue); }
+	bool GetTimersViewShowVerseSamplingTimers() const { return bTimersViewShowVerseSamplingTimers; }
+	void SetTimersViewShowVerseSamplingTimers(bool InValue) { bTimersViewShowVerseSamplingTimers = InValue; }
+	void SetAndSaveTimersViewShowVerseSamplingTimers(bool InValue) { SET_AND_SAVE(bTimersViewShowVerseSamplingTimers, InValue); }
 
-	bool GetTimersViewShowCpuEvents() const { return bTimersViewShowCpuTimers; }
-	void SetTimersViewShowCpuEvents(bool InValue) { bTimersViewShowCpuTimers = InValue; }
-	void SetAndSaveTimersViewShowCpuEvents(bool InValue) { SET_AND_SAVE(bTimersViewShowCpuTimers, InValue); }
+	bool GetTimersViewShowCpuTimers() const { return bTimersViewShowCpuTimers; }
+	void SetTimersViewShowCpuTimers(bool InValue) { bTimersViewShowCpuTimers = InValue; }
+	void SetAndSaveTimersViewShowCpuTimers(bool InValue) { SET_AND_SAVE(bTimersViewShowCpuTimers, InValue); }
+
+	bool GetTimersViewShowCpuSamplingTimers() const { return bTimersViewShowCpuSamplingTimers; }
+	void SetTimersViewShowCpuSamplingTimers(bool InValue) { bTimersViewShowCpuSamplingTimers = InValue; }
+	void SetAndSaveTimersViewShowCpuSamplingTimers(bool InValue) { SET_AND_SAVE(bTimersViewShowCpuSamplingTimers, InValue); }
 
 	bool GetTimersViewShowZeroCountTimers() const { return bTimersViewShowZeroCountTimers; }
 	void SetTimersViewShowZeroCountTimers(bool InValue) { bTimersViewShowZeroCountTimers = InValue; }
@@ -614,15 +618,18 @@ private:
 	/** The grouping mode for the timers panel. */
 	int32 TimersViewGroupingMode = 3; // ByType
 
-	/** If enabled, GPU timers will be displayed in the Timing View. */
+	/** If enabled, GPU Scope timers will be displayed in the Timing View. */
 	bool bTimersViewShowGpuTimers = true;
 
-	/** If enabled, Verse timers will be displayed in the Timing View. */
-	bool bTimersViewShowVerseTimers = true;
+	/** If enabled, Verse Sampling timers will be displayed in the Timing View. */
+	bool bTimersViewShowVerseSamplingTimers = true;
 
-	/** If enabled, CPU timers will be displayed in the Timing View. */
+	/** If enabled, CPU Scope timers will be displayed in the Timing View. */
 	bool bTimersViewShowCpuTimers = true;
 
+	/** If enabled, CPU Sampling timers will be displayed in the Timing View. */
+	bool bTimersViewShowCpuSamplingTimers = true;
+
 	/** If enabled, timers with no instances in the selected interval will still be displayed in the Timers View. */
 	bool bTimersViewShowZeroCountTimers = true;
 
diff --git a/Engine/Source/Developer/TraceInsights/Private/Insights/TimingProfiler/Tracks/CpuStackSampleTimingTrack.cpp b/Engine/Source/Developer/TraceInsights/Private/Insights/TimingProfiler/Tracks/CpuStackSampleTimingTrack.cpp
new file mode 100644
--- /dev/null
+++ b/Engine/Source/Developer/TraceInsights/Private/Insights/TimingProfiler/Tracks/CpuStackSampleTimingTrack.cpp
@@ -0,0 +1,47 @@
+// Copyright Epic Games, Inc. All Rights Reserved.
+
+#include "CpuStackSampleTimingTrack.h"
+
+// TraceServices
+#include "Common/ProviderLock.h"
+#include "TraceServices/Model/AnalysisSession.h"
+#include "TraceServices/Model/StackSamples.h"
+
+// TraceInsights
+#include "Insights/InsightsManager.h"
+
+#define LOCTEXT_NAMESPACE "UE::Insights::TimingProfiler::CpuStackSampleTimingTrack"
+
+namespace UE::Insights::TimingProfiler
+{
+
+////////////////////////////////////////////////////////////////////////////////////////////////////
+// FCpuStackSampleTimingTrack
+////////////////////////////////////////////////////////////////////////////////////////////////////
+
+INSIGHTS_IMPLEMENT_RTTI(FCpuStackSampleTimingTrack)
+
+////////////////////////////////////////////////////////////////////////////////////////////////////
+
+bool FCpuStackSampleTimingTrack::ReadTimeline(TFunctionRef<void(const TraceServices::ITimingProfilerProvider::Timeline&)> Callback) const
+{
+	TSharedPtr<const TraceServices::IAnalysisSession> Session = FInsightsManager::Get()->GetSession();
+	const TraceServices::IStackSamplesProvider* StackSamplesProvider = TraceServices::ReadStackSamplesProvider(*Session.Get());
+	if (StackSamplesProvider)
+	{
+		TraceServices::FProviderReadScopeLock _(*StackSamplesProvider);
+		const TraceServices::ITimingProfilerProvider::Timeline* Timeline = StackSamplesProvider->GetTimeline(GetTimelineIndex());
+		if (Timeline)
+		{
+			Callback(*Timeline);
+			return true;
+		}
+	}
+	return false;
+}
+
+////////////////////////////////////////////////////////////////////////////////////////////////////
+
+} // namespace UE::Insights::TimingProfiler
+
+#undef LOCTEXT_NAMESPACE
diff --git a/Engine/Source/Developer/TraceInsights/Private/Insights/TimingProfiler/Tracks/CpuStackSampleTimingTrack.h b/Engine/Source/Developer/TraceInsights/Private/Insights/TimingProfiler/Tracks/CpuStackSampleTimingTrack.h
new file mode 100644
--- /dev/null
+++ b/Engine/Source/Developer/TraceInsights/Private/Insights/TimingProfiler/Tracks/CpuStackSampleTimingTrack.h
@@ -0,0 +1,31 @@
+// Copyright Epic Games, Inc. All Rights Reserved.
+
+#pragma once
+
+#include "CoreTypes.h"
+
+// TraceInsights
+#include "Insights/TimingProfiler/Tracks/ThreadTimingTrackPrivate.h"
+
+namespace UE::Insights::TimingProfiler
+{
+
+////////////////////////////////////////////////////////////////////////////////////////////////////
+
+class FCpuStackSampleTimingTrack : public FThreadTimingTrackImpl
+{
+	INSIGHTS_DECLARE_RTTI(FCpuStackSampleTimingTrack, FThreadTimingTrackImpl)
+
+public:
+	explicit FCpuStackSampleTimingTrack(FThreadTimingSharedState& InSharedState, const FString& InName, uint32 InSystemThreadId)
+		: FThreadTimingTrackImpl(InSharedState, InName, nullptr, InSystemThreadId, InSystemThreadId)
+	{
+	}
+
+protected:
+	virtual bool ReadTimeline(TFunctionRef<void(const TraceServices::ITimingProfilerProvider::Timeline&)> Callback) const override;
+};
+
+////////////////////////////////////////////////////////////////////////////////////////////////////
+
+} // namespace UE::Insights::TimingProfiler
diff --git a/Engine/Source/Developer/TraceInsights/Private/Insights/TimingProfiler/ViewModels/ThreadTimingSharedState.cpp b/Engine/Source/Developer/TraceInsights/Private/Insights/TimingProfiler/ViewModels/ThreadTimingSharedState.cpp
--- a/Engine/Source/Developer/TraceInsights/Private/Insights/TimingProfiler/ViewModels/ThreadTimingSharedState.cpp
+++ b/Engine/Source/Developer/TraceInsights/Private/Insights/TimingProfiler/ViewModels/ThreadTimingSharedState.cpp
@@ -8,8 +8,10 @@
 #include "HAL/LowLevelMemTracker.h"
 
 // TraceServices
+#include "Common/ProviderLock.h"
 #include "TraceServices/Model/AnalysisSession.h"
 #include "TraceServices/Model/LoadTimeProfiler.h"
+#include "TraceServices/Model/StackSamples.h"
 #include "TraceServices/Model/TimingProfiler.h"
 #include "TraceServices/Model/Threads.h"
 
@@ -20,6 +22,7 @@
 #include "Insights/InsightsStyle.h"
 #include "Insights/IUnrealInsightsModule.h"
 #include "Insights/ITimingViewSession.h"
+#include "Insights/TimingProfiler/Tracks/CpuStackSampleTimingTrack.h"
 #include "Insights/TimingProfiler/Tracks/CpuTimingTrack.h"
 #include "Insights/TimingProfiler/Tracks/GpuTimingTrack.h"
 #include "Insights/TimingProfiler/Tracks/VerseTimingTrack.h"
@@ -403,6 +406,8 @@ void FThreadTimingSharedState::Tick(Timing::ITimingViewSession& InSession, const
 		return;
 	}
 
+	bool bTracksOrderChanged = false;
+
 	const TraceServices::ITimingProfilerProvider* TimingProfilerProvider = TraceServices::ReadTimingProfilerProvider(InAnalysisSession);
 	const TraceServices::ILoadTimeProfilerProvider* LoadTimeProfilerProvider = TraceServices::ReadLoadTimeProfilerProvider(InAnalysisSession);
 
@@ -410,8 +415,8 @@ void FThreadTimingSharedState::Tick(Timing::ITimingViewSession& InSession, const
 	{
 		TraceServices::FAnalysisSessionReadScope SessionReadScope(InAnalysisSession);
 
-		const uint64 CurrentTimingProfilerTimelineCount = TimingProfilerProvider->GetTimelineCount();
-		const uint64 CurrentLoadTimeProfilerTimelineCount = (LoadTimeProfilerProvider) ? LoadTimeProfilerProvider->GetTimelineCount() : 0;
+		const uint32 CurrentTimingProfilerTimelineCount = TimingProfilerProvider->GetTimelineCount();
+		const uint32 CurrentLoadTimeProfilerTimelineCount = (LoadTimeProfilerProvider) ? LoadTimeProfilerProvider->GetTimelineCount() : 0;
 
 		if (CurrentTimingProfilerTimelineCount != TimingProfilerTimelineCount ||
 			CurrentLoadTimeProfilerTimelineCount != LoadTimeProfilerTimelineCount)
@@ -447,7 +452,6 @@ void FThreadTimingSharedState::Tick(Timing::ITimingViewSession& InSession, const
 				}
 			}
 
-			bool bTracksOrderChanged = false;
 			int32 GpuTrackOrder = FTimingTrackOrder::Gpu + 100;
 			int32 CpuTrackOrder = FTimingTrackOrder::Cpu;
 
@@ -512,7 +516,9 @@ void FThreadTimingSharedState::Tick(Timing::ITimingViewSession& InSession, const
 
 			// Iterate through threads.
 			const TraceServices::IThreadProvider& ThreadProvider = TraceServices::ReadThreadProvider(InAnalysisSession);
-			ThreadProvider.EnumerateThreads([this, &InSession, &bTracksOrderChanged, &CpuTrackOrder, TimingProfilerProvider, LoadTimeProfilerProvider](const TraceServices::FThreadInfo& ThreadInfo)
+			ThreadProvider.EnumerateThreads(
+				[this, &InSession, &bTracksOrderChanged, &CpuTrackOrder, TimingProfilerProvider, LoadTimeProfilerProvider]
+				(const TraceServices::FThreadInfo& ThreadInfo)
 			{
 				// Check if this thread is part of a group?
 				const TCHAR* GroupName = ThreadInfo.GroupName;
@@ -592,13 +598,58 @@ void FThreadTimingSharedState::Tick(Timing::ITimingViewSession& InSession, const
 				static_assert(OrderIncrement >= 1, "Order group range too small");
 				CpuTrackOrder += OrderIncrement;
 			});
+		}
+	}
 
-			if (bTracksOrderChanged)
-			{
-				InSession.InvalidateScrollableTracksOrder();
-			}
+	if (const TraceServices::IStackSamplesProvider* StackSamplesProvider = TraceServices::ReadStackSamplesProvider(InAnalysisSession))
+	{
+		TraceServices::FProviderReadScopeLock _(*StackSamplesProvider);
+
+		const uint32 CurrentTimelineCount = StackSamplesProvider->GetTimelineCount();
+
+		if (CurrentTimelineCount != CpuStackSamplingTimelineCount)
+		{
+			CpuStackSamplingTimelineCount = CurrentTimelineCount;
+
+			LLM_SCOPE_BYTAG(Insights);
+
+			int32 CpuStackSampleTrackOrder = FTimingTrackOrder::Cpu + FTimingTrackOrder::GroupRange;
+
+			StackSamplesProvider->EnumerateTimelines(
+				[this, &InSession, &bTracksOrderChanged, &CpuStackSampleTrackOrder]
+				(const TraceServices::FStackSampleTimeline& Timeline, const TraceServices::FStackSampleThread& Thread)
+				{
+					TSharedPtr<FCpuStackSampleTimingTrack>* TrackPtrPtr = CpuStackSampleTracks.Find(Thread.SystemThreadId);
+					if (TrackPtrPtr == nullptr)
+					{
+						FString TrackName(Thread.Name && *Thread.Name ? Thread.Name : FString::Printf(TEXT("Sampling Thread 0x%X"), Thread.SystemThreadId));
+						TSharedPtr<FCpuStackSampleTimingTrack> Track = MakeShared<FCpuStackSampleTimingTrack>(*this, TrackName, Thread.SystemThreadId);
+						Track->SetOrder(CpuStackSampleTrackOrder);
+						CpuStackSampleTracks.Add(Thread.SystemThreadId, Track);
+						Track->SetVisibilityFlag(true);
+						InSession.AddScrollableTrack(Track);
+					}
+					else
+					{
+						TSharedPtr<FCpuStackSampleTimingTrack> Track = *TrackPtrPtr;
+						if (Track->GetOrder() != CpuStackSampleTrackOrder)
+						{
+							Track->SetOrder(CpuStackSampleTrackOrder);
+							bTracksOrderChanged = true;
+						}
+					}
+
+					constexpr int32 OrderIncrement = FTimingTrackOrder::GroupRange / 1000; // distribute max 1000 tracks in the order group range
+					static_assert(OrderIncrement >= 1, "Order group range too small");
+					CpuStackSampleTrackOrder += OrderIncrement;
+				});
 		}
 	}
+
+	if (bTracksOrderChanged)
+	{
+		InSession.InvalidateScrollableTracksOrder();
+	}
 }
 
 ////////////////////////////////////////////////////////////////////////////////////////////////////
diff --git a/Engine/Source/Developer/TraceInsights/Private/Insights/TimingProfiler/ViewModels/ThreadTimingSharedState.h b/Engine/Source/Developer/TraceInsights/Private/Insights/TimingProfiler/ViewModels/ThreadTimingSharedState.h
--- a/Engine/Source/Developer/TraceInsights/Private/Insights/TimingProfiler/ViewModels/ThreadTimingSharedState.h
+++ b/Engine/Source/Developer/TraceInsights/Private/Insights/TimingProfiler/ViewModels/ThreadTimingSharedState.h
@@ -23,6 +23,7 @@ namespace UE::Insights::TimingProfiler
 {
 
 class FCpuTimingTrack;
+class FCpuStackSampleTimingTrack;
 class FGpuTimingTrack;
 class FGpuQueueTimingTrack;
 class FGpuQueueWorkTimingTrack;
@@ -266,19 +267,23 @@ private:
 	TSharedPtr<FGpuTimingTrack> OldGpu1Track;
 	TSharedPtr<FGpuTimingTrack> OldGpu2Track;
 
-	/** Maps GPU queue id to track pointer. */
+	/** Maps GPU queue id to track pointer (GPU Queue timing tracks). */
 	TMap<uint32, TSharedPtr<FGpuQueueTimingTrack>> GpuTracks;
 
 	TSharedPtr<FVerseTimingTrack> VerseSamplingTrack;
 
-	/** Maps CPU thread id to track pointer. */
+	/** Maps CPU thread id to track pointer (CPU timing tracks). */
 	TMap<uint32, TSharedPtr<FCpuTimingTrack>> CpuTracks;
 
-	/** Maps thread group name to thread group info. */
+	/** Maps CPU thread group name to thread group info. */
 	TMap<const TCHAR*, FThreadGroup, FDefaultSetAllocator, TStringPointerMapKeyFuncs_DEPRECATED<const TCHAR*, FThreadGroup>> ThreadGroups;
 
-	uint64 TimingProfilerTimelineCount = 0;
-	uint64 LoadTimeProfilerTimelineCount = 0;
+	/** Maps system thread id to track pointer (CPU Stack Sample timing tracks). */
+	TMap<uint32, TSharedPtr<FCpuStackSampleTimingTrack>> CpuStackSampleTracks;
+
+	uint32 TimingProfilerTimelineCount = 0;
+	uint32 LoadTimeProfilerTimelineCount = 0;
+	uint32 CpuStackSamplingTimelineCount = 0;
 
 	TSharedRef<ITimingViewSettings> TimingViewSettings;
 };
diff --git a/Engine/Source/Developer/TraceInsights/Private/Insights/TimingProfiler/ViewModels/TimerNode.cpp b/Engine/Source/Developer/TraceInsights/Private/Insights/TimingProfiler/ViewModels/TimerNode.cpp
--- a/Engine/Source/Developer/TraceInsights/Private/Insights/TimingProfiler/ViewModels/TimerNode.cpp
+++ b/Engine/Source/Developer/TraceInsights/Private/Insights/TimingProfiler/ViewModels/TimerNode.cpp
@@ -19,18 +19,51 @@ namespace UE::Insights::TimingProfiler
 
 INSIGHTS_IMPLEMENT_RTTI(FTimerNode)
 
-const FName FTimerNode::GpuGroup(TEXT("GPU"));
-const FName FTimerNode::VerseGroup(TEXT("Verse"));
-const FName FTimerNode::CpuGroup(TEXT("CPU"));
+static FName GetGpuScopeGroup()
+{
+	static FName Group(TEXT("GPU"));
+	return Group;
+}
+
+static FName GetVerseSamplingGroup()
+{
+	static FName Group(TEXT("Verse"));
+	return Group;
+}
+
+static FName GetCpuScopeGroup()
+{
+	static FName Group(TEXT("CPU"));
+	return Group;
+}
+
+static FName GetCpuSamplingGroup()
+{
+	static FName Group(TEXT("CPU Sampling"));
+	return Group;
+}
+
+////////////////////////////////////////////////////////////////////////////////////////////////////
+
+FName MakeSafeFName(const TCHAR* InName)
+{
+	int32 Len = FCString::Strlen(InName);
+	if (Len >= NAME_SIZE)
+	{
+		return FName(NAME_SIZE - 1, InName + (Len - NAME_SIZE + 1));
+	}
+	return FName(Len, InName);
+}
 
 ////////////////////////////////////////////////////////////////////////////////////////////////////
 
 FTimerNode::FTimerNode(uint32 InTimerId, const TCHAR* InName, ETimerNodeType InType, bool bInIsGroup)
-	: FBaseTreeNode(FName(InName), bInIsGroup)
+	: FBaseTreeNode(MakeSafeFName(InName), bInIsGroup)
 	, TimerId(InTimerId)
-	, MetaGroupName(InType == ETimerNodeType::GpuScope ? GpuGroup :
-					InType == ETimerNodeType::VerseScope ? VerseGroup :
-					InType == ETimerNodeType::CpuScope ? CpuGroup :
+	, MetaGroupName(InType == ETimerNodeType::GpuScope ? GetGpuScopeGroup() :
+					InType == ETimerNodeType::VerseSampling ? GetVerseSamplingGroup() :
+					InType == ETimerNodeType::CpuScope ? GetCpuScopeGroup() :
+					InType == ETimerNodeType::CpuSampling ? GetCpuSamplingGroup() :
 					NAME_None)
 	, Type(InType)
 {
diff --git a/Engine/Source/Developer/TraceInsights/Private/Insights/TimingProfiler/ViewModels/TimerNode.h b/Engine/Source/Developer/TraceInsights/Private/Insights/TimingProfiler/ViewModels/TimerNode.h
--- a/Engine/Source/Developer/TraceInsights/Private/Insights/TimingProfiler/ViewModels/TimerNode.h
+++ b/Engine/Source/Developer/TraceInsights/Private/Insights/TimingProfiler/ViewModels/TimerNode.h
@@ -21,12 +21,15 @@ enum class ETimerNodeType
 	/** The TimerNode is a GPU Scope timer. */
 	GpuScope,
 
-	/** The TimerNode is a Verse Scope timer. */
-	VerseScope,
+	/** The TimerNode is a Verse Sampling timer. */
+	VerseSampling,
 
 	/** The TimerNode is a CPU Scope timer. */
 	CpuScope,
 
+	/** The TimerNode is a CPU Sampling timer. */
+	CpuSampling,
+
 	/** The TimerNode is a group node. */
 	Group,
 
@@ -61,10 +64,6 @@ class FTimerNode : public FBaseTreeNode
 public:
 	static constexpr uint32 InvalidTimerId = uint32(-1);
 
-	static const FName GpuGroup;
-	static const FName VerseGroup;
-	static const FName CpuGroup;
-
 public:
 	/** Initialization constructor for the timer node. */
 	FTimerNode(uint32 InTimerId, const TCHAR* InName, ETimerNodeType InType, bool bInIsGroup);
diff --git a/Engine/Source/Developer/TraceInsights/Private/Insights/TimingProfiler/ViewModels/TimerNodeHelper.cpp b/Engine/Source/Developer/TraceInsights/Private/Insights/TimingProfiler/ViewModels/TimerNodeHelper.cpp
--- a/Engine/Source/Developer/TraceInsights/Private/Insights/TimingProfiler/ViewModels/TimerNodeHelper.cpp
+++ b/Engine/Source/Developer/TraceInsights/Private/Insights/TimingProfiler/ViewModels/TimerNodeHelper.cpp
@@ -16,12 +16,13 @@ namespace UE::Insights::TimingProfiler
 
 FText TimerNodeTypeHelper::ToText(const ETimerNodeType NodeType)
 {
-	static_assert(static_cast<int>(ETimerNodeType::InvalidOrMax) == 4, "Not all cases are handled in switch below!?");
+	static_assert(static_cast<int>(ETimerNodeType::InvalidOrMax) == 5, "Not all cases are handled in switch below!?");
 	switch (NodeType)
 	{
-		case ETimerNodeType::GpuScope:		return LOCTEXT("Timer_Name_Gpu", "GPU");
-		case ETimerNodeType::VerseScope:	return LOCTEXT("Timer_Name_Verse", "Verse");
-		case ETimerNodeType::CpuScope:		return LOCTEXT("Timer_Name_Cpu", "CPU");
+		case ETimerNodeType::GpuScope:		return LOCTEXT("Timer_Name_GpuScope", "GPU");
+		case ETimerNodeType::VerseSampling:	return LOCTEXT("Timer_Name_VerseSampling", "Verse");
+		case ETimerNodeType::CpuScope:		return LOCTEXT("Timer_Name_CpuScope", "CPU");
+		case ETimerNodeType::CpuSampling:	return LOCTEXT("Timer_Name_CpuSampling", "Sampling");
 		case ETimerNodeType::Group:			return LOCTEXT("Timer_Name_Group", "Group");
 		default:							return LOCTEXT("InvalidOrMax", "InvalidOrMax");
 	}
@@ -31,12 +32,13 @@ FText TimerNodeTypeHelper::ToText(const ETimerNodeType NodeType)
 
 FText TimerNodeTypeHelper::ToDescription(const ETimerNodeType NodeType)
 {
-	static_assert(static_cast<int>(ETimerNodeType::InvalidOrMax) == 4, "Not all cases are handled in switch below!?");
+	static_assert(static_cast<int>(ETimerNodeType::InvalidOrMax) == 5, "Not all cases are handled in switch below!?");
 	switch (NodeType)
 	{
-		case ETimerNodeType::GpuScope:		return LOCTEXT("Timer_Desc_Gpu", "GPU scope timer");
-		case ETimerNodeType::VerseScope:	return LOCTEXT("Timer_Desc_Verse", "Verse scope timer");
-		case ETimerNodeType::CpuScope:		return LOCTEXT("Timer_Desc_Cpu", "CPU scope timer");
+		case ETimerNodeType::GpuScope:		return LOCTEXT("Timer_Desc_GpuScope", "GPU Scope timer");
+		case ETimerNodeType::VerseSampling:	return LOCTEXT("Timer_Desc_VerseSampling", "Verse Sampling timer");
+		case ETimerNodeType::CpuScope:		return LOCTEXT("Timer_Desc_CpuScope", "CPU Scope timer");
+		case ETimerNodeType::CpuSampling:	return LOCTEXT("Timer_Desc_CpuSampling", "CPU Sampling timer");
 		case ETimerNodeType::Group:			return LOCTEXT("Timer_Desc_Group", "Group timer node");
 		default:							return LOCTEXT("InvalidOrMax", "InvalidOrMax");
 	}
@@ -53,12 +55,13 @@ const FSlateBrush* TimerNodeTypeHelper::GetIconForGroup()
 
 const FSlateBrush* TimerNodeTypeHelper::GetIconForTimerNodeType(const ETimerNodeType NodeType)
 {
-	static_assert(static_cast<int>(ETimerNodeType::InvalidOrMax) == 4, "Not all cases are handled in switch below!?");
+	static_assert(static_cast<int>(ETimerNodeType::InvalidOrMax) == 5, "Not all cases are handled in switch below!?");
 	switch (NodeType)
 	{
 		case ETimerNodeType::GpuScope:		return FInsightsStyle::GetBrush("Icons.GpuTimer.TreeItem");
-		case ETimerNodeType::VerseScope:	return FInsightsStyle::GetBrush("Icons.CpuTimer.TreeItem");
+		case ETimerNodeType::VerseSampling:	return FInsightsStyle::GetBrush("Icons.CpuTimer.TreeItem");
 		case ETimerNodeType::CpuScope:		return FInsightsStyle::GetBrush("Icons.CpuTimer.TreeItem");
+		case ETimerNodeType::CpuSampling:	return FInsightsStyle::GetBrush("Icons.CpuTimer.TreeItem");
 		case ETimerNodeType::Group:			return FInsightsStyle::GetBrush("Icons.Group.TreeItem");
 		default:							return nullptr;
 	}
diff --git a/Engine/Source/Developer/TraceInsights/Private/Insights/TimingProfiler/ViewModels/TimingExporter.cpp b/Engine/Source/Developer/TraceInsights/Private/Insights/TimingProfiler/ViewModels/TimingExporter.cpp
--- a/Engine/Source/Developer/TraceInsights/Private/Insights/TimingProfiler/ViewModels/TimingExporter.cpp
+++ b/Engine/Source/Developer/TraceInsights/Private/Insights/TimingProfiler/ViewModels/TimingExporter.cpp
@@ -284,14 +284,17 @@ int32 FTimingExporter::ExportTimersAsText(const FString& Filename, FExportTimers
 			Writer.AppendSeparator();
 			switch (Timer.Type)
 			{
-			case TraceServices::ETimingProfilerTimerType::CPU:
+			case TraceServices::ETimingProfilerTimerType::CpuScope:
 				StringBuilder.Append(UTF8TEXT("CPU"));
 				break;
-			case TraceServices::ETimingProfilerTimerType::GPU:
+			case TraceServices::ETimingProfilerTimerType::GpuScope:
 				StringBuilder.Append(UTF8TEXT("GPU"));
 				break;
-			case TraceServices::ETimingProfilerTimerType::Verse:
-				StringBuilder.Append(UTF8TEXT("Verse"));
+			case TraceServices::ETimingProfilerTimerType::VerseSampling:
+				StringBuilder.Append(UTF8TEXT("VerseSampling"));
+				break;
+			case TraceServices::ETimingProfilerTimerType::CpuSampling:
+				StringBuilder.Append(UTF8TEXT("CpuSampling"));
 				break;
 			default:
 				StringBuilder.Append(UTF8TEXT("Unknown"));
diff --git a/Engine/Source/Developer/TraceInsights/Private/Insights/TimingProfiler/Widgets/STimerTreeView.cpp b/Engine/Source/Developer/TraceInsights/Private/Insights/TimingProfiler/Widgets/STimerTreeView.cpp
--- a/Engine/Source/Developer/TraceInsights/Private/Insights/TimingProfiler/Widgets/STimerTreeView.cpp
+++ b/Engine/Source/Developer/TraceInsights/Private/Insights/TimingProfiler/Widgets/STimerTreeView.cpp
@@ -1524,9 +1524,10 @@ FTimerNodePtr STimerTreeView::CreateTimerNodeRec(const TraceServices::FTimingPro
 	}
 
 	const ETimerNodeType Type =
-		Node.Timer->Type == TraceServices::ETimingProfilerTimerType::GPU ? ETimerNodeType::GpuScope :
-		Node.Timer->Type == TraceServices::ETimingProfilerTimerType::Verse ? ETimerNodeType::VerseScope :
-		Node.Timer->Type == TraceServices::ETimingProfilerTimerType::CPU ? ETimerNodeType::CpuScope :
+		Node.Timer->Type == TraceServices::ETimingProfilerTimerType::GpuScope ? ETimerNodeType::GpuScope :
+		Node.Timer->Type == TraceServices::ETimingProfilerTimerType::VerseSampling ? ETimerNodeType::VerseSampling :
+		Node.Timer->Type == TraceServices::ETimingProfilerTimerType::CpuScope ? ETimerNodeType::CpuScope :
+		Node.Timer->Type == TraceServices::ETimingProfilerTimerType::CpuSampling ? ETimerNodeType::CpuSampling :
 		ETimerNodeType::InvalidOrMax;
 
 	TSharedPtr<STimingProfilerWindow> ProfilerWindow = FTimingProfilerManager::Get()->GetProfilerWindow();
@@ -1538,7 +1539,7 @@ FTimerNodePtr STimerTreeView::CreateTimerNodeRec(const TraceServices::FTimingPro
 	{
 		TimerNodePtr->SetColor(ColorOverride);
 	}
-	
+
 	TimerNodePtr->SetAddedToGraphsFlag(bIsTimerAddedToGraphs);
 
 	TraceServices::FTimingProfilerAggregatedStats AggregatedStats;
diff --git a/Engine/Source/Developer/TraceInsights/Private/Insights/TimingProfiler/Widgets/STimersView.cpp b/Engine/Source/Developer/TraceInsights/Private/Insights/TimingProfiler/Widgets/STimersView.cpp
--- a/Engine/Source/Developer/TraceInsights/Private/Insights/TimingProfiler/Widgets/STimersView.cpp
+++ b/Engine/Source/Developer/TraceInsights/Private/Insights/TimingProfiler/Widgets/STimersView.cpp
@@ -27,8 +27,10 @@
 #include "Widgets/Views/STableViewBase.h"
 
 // TraceServices
+#include "Common/ProviderLock.h"
 #include "TraceServices/Model/Threads.h"
 #include "TraceServices/Model/TimingProfiler.h"
+#include "TraceServices/Model/StackSamples.h"
 
 // TraceInsightsCore
 #include "InsightsCore/Common/Stopwatch.h"
@@ -183,7 +185,7 @@ public:
 			"Navigates to and selects the timing event instance with the minimum duration, for the selected timer.",
 			EUserInterfaceActionType::Button,
 			FInputChord());
-		
+
 		UI_COMMAND(Command_FindFirstInstanceInSelection,
 			"First Instance in Selection",
 			"Navigates to and selects the first timing event instance for the selected timer, in the selected time range.",
@@ -347,7 +349,7 @@ void STimersView::Construct(const FArguments& InArgs)
 					.ToolTipText(LOCTEXT("FilterSearchHint", "Type here to search timer or group."))
 				]
 
-				// Filter by type (GPU)
+				// Filter by type (GPU Scope)
 				+ SHorizontalBox::Slot()
 				.Padding(4.0f, 0.0f, 4.0f, 0.0f)
 				.AutoWidth()
@@ -356,16 +358,16 @@ void STimersView::Construct(const FArguments& InArgs)
 					GetToggleButtonForTimerType(ETimerNodeType::GpuScope)
 				]
 
-				// Filter by type (Verse)
+				// Filter by type (Verse Sampling)
 				+ SHorizontalBox::Slot()
 				.Padding(4.0f, 0.0f, 4.0f, 0.0f)
 				.AutoWidth()
 				.VAlign(VAlign_Center)
 				[
-					GetToggleButtonForTimerType(ETimerNodeType::VerseScope)
+					GetToggleButtonForTimerType(ETimerNodeType::VerseSampling)
 				]
 
-				// Filter by type (CPU)
+				// Filter by type (CPU Scope)
 				+ SHorizontalBox::Slot()
 				.Padding(4.0f, 0.0f, 4.0f, 0.0f)
 				.AutoWidth()
@@ -374,6 +376,15 @@ void STimersView::Construct(const FArguments& InArgs)
 					GetToggleButtonForTimerType(ETimerNodeType::CpuScope)
 				]
 
+				// Filter by type (CPU Sampling)
+				+ SHorizontalBox::Slot()
+				.Padding(4.0f, 0.0f, 4.0f, 0.0f)
+				.AutoWidth()
+				.VAlign(VAlign_Center)
+				[
+					GetToggleButtonForTimerType(ETimerNodeType::CpuSampling)
+				]
+
 				// Filter out timers with zero instance count
 				+ SHorizontalBox::Slot()
 				.Padding(2.0f)
@@ -1283,7 +1294,7 @@ void STimersView::TreeView_FindMenu(FMenuBuilder& MenuBuilder)
 		TAttribute<FText>(),
 		FSlateIcon(FInsightsStyle::Get().GetStyleSetName(), "Icons.FindLast")
 	);
-	
+
 	MenuBuilder.AddMenuEntry
 	(
 		FTimersViewCommands::Get().Command_FindMaxInstanceInSelection,
@@ -2718,20 +2729,40 @@ void STimersView::Tick(const FGeometry& AllottedGeometry, const double InCurrent
 
 	Aggregator->Tick(Session, InCurrentTime, InDeltaTime, [this]() { FinishAggregation(); });
 
-	if (Session.IsValid() && TraceServices::ReadTimingProfilerProvider(*Session.Get()))
+	if (Session.IsValid())
 	{
-		TraceServices::FAnalysisSessionReadScope SessionReadScope(*Session.Get());
-		const TraceServices::ITimingProfilerProvider& TimingProfilerProvider = *TraceServices::ReadTimingProfilerProvider(*Session.Get());
+		if (const TraceServices::ITimingProfilerProvider* TimingProfilerProvider = TraceServices::ReadTimingProfilerProvider(*Session.Get()))
+		{
+			TraceServices::FAnalysisSessionReadScope SessionReadScope(*Session.Get());
+
+			bTimerTypeIsEnabled[static_cast<int>(ETimerNodeType::GpuScope)] = TimingProfilerProvider->HasGpuTiming();
+			bTimerTypeIsEnabled[static_cast<int>(ETimerNodeType::VerseSampling)] = TimingProfilerProvider->HasVerseTiming();
+			bTimerTypeIsEnabled[static_cast<int>(ETimerNodeType::CpuScope)] = TimingProfilerProvider->HasCpuTiming();
+		}
+		else
+		{
+			bTimerTypeIsEnabled[static_cast<int>(ETimerNodeType::GpuScope)] = false;
+			bTimerTypeIsEnabled[static_cast<int>(ETimerNodeType::VerseSampling)] = false;
+			bTimerTypeIsEnabled[static_cast<int>(ETimerNodeType::CpuScope)] = false;
+		}
 
-		bTimerTypeIsEnabled[static_cast<int>(ETimerNodeType::GpuScope)] = TimingProfilerProvider.HasGpuTiming();
-		bTimerTypeIsEnabled[static_cast<int>(ETimerNodeType::VerseScope)] = TimingProfilerProvider.HasVerseTiming();
-		bTimerTypeIsEnabled[static_cast<int>(ETimerNodeType::CpuScope)] = TimingProfilerProvider.HasCpuTiming();
+		if (const TraceServices::IStackSamplesProvider* StackSamplesProvider = TraceServices::ReadStackSamplesProvider(*Session.Get()))
+		{
+			TraceServices::FProviderReadScopeLock _(*StackSamplesProvider);
+
+			bTimerTypeIsEnabled[static_cast<int>(ETimerNodeType::CpuSampling)] = (StackSamplesProvider->GetTimelineCount() > 0);
+		}
+		else
+		{
+			bTimerTypeIsEnabled[static_cast<int>(ETimerNodeType::CpuSampling)] = false;
+		}
 	}
 	else
 	{
 		bTimerTypeIsEnabled[static_cast<int>(ETimerNodeType::GpuScope)] = false;
-		bTimerTypeIsEnabled[static_cast<int>(ETimerNodeType::VerseScope)] = false;
+		bTimerTypeIsEnabled[static_cast<int>(ETimerNodeType::VerseSampling)] = false;
 		bTimerTypeIsEnabled[static_cast<int>(ETimerNodeType::CpuScope)] = false;
+		bTimerTypeIsEnabled[static_cast<int>(ETimerNodeType::CpuSampling)] = false;
 	}
 }
 
@@ -2776,9 +2807,10 @@ void STimersView::RebuildTree(bool bResync)
 				const TraceServices::FTimingProfilerTimer& Timer = *(TimerReader->GetTimer(TimerIndex));
 				ensure(Timer.Id == TimerIndex);
 				const ETimerNodeType Type =
-					Timer.Type == TraceServices::ETimingProfilerTimerType::GPU ? ETimerNodeType::GpuScope :
-					Timer.Type == TraceServices::ETimingProfilerTimerType::Verse ? ETimerNodeType::VerseScope :
-					Timer.Type == TraceServices::ETimingProfilerTimerType::CPU ? ETimerNodeType::CpuScope :
+					Timer.Type == TraceServices::ETimingProfilerTimerType::GpuScope ? ETimerNodeType::GpuScope :
+					Timer.Type == TraceServices::ETimingProfilerTimerType::VerseSampling ? ETimerNodeType::VerseSampling :
+					Timer.Type == TraceServices::ETimingProfilerTimerType::CpuScope ? ETimerNodeType::CpuScope :
+					Timer.Type == TraceServices::ETimingProfilerTimerType::CpuSampling ? ETimerNodeType::CpuSampling :
 					ETimerNodeType::InvalidOrMax;
 				FTimerNodePtr TimerNodePtr = MakeShared<FTimerNode>(Timer.Id, Timer.Name, Type, false);
 				TimerNodePtr->SetDefaultSortOrder(TimerIndex + 1);
@@ -3951,9 +3983,10 @@ void STimersView::SaveSettings()
 
 	Settings.SetTimersViewMode(static_cast<int32>(ModeFrameType));
 	Settings.SetTimersViewGroupingMode(static_cast<int32>(GroupingMode));
-	Settings.SetTimersViewShowGpuEvents(static_cast<int32>(bTimerTypeIsVisible[static_cast<int32>(ETimerNodeType::GpuScope)]));
-	Settings.SetTimersViewShowVerseEvents(static_cast<int32>(bTimerTypeIsVisible[static_cast<int32>(ETimerNodeType::VerseScope)]));
-	Settings.SetTimersViewShowCpuEvents(static_cast<int32>(bTimerTypeIsVisible[static_cast<int32>(ETimerNodeType::CpuScope)]));
+	Settings.SetTimersViewShowGpuTimers(static_cast<int32>(bTimerTypeIsVisible[static_cast<int32>(ETimerNodeType::GpuScope)]));
+	Settings.SetTimersViewShowVerseSamplingTimers(static_cast<int32>(bTimerTypeIsVisible[static_cast<int32>(ETimerNodeType::VerseSampling)]));
+	Settings.SetTimersViewShowCpuTimers(static_cast<int32>(bTimerTypeIsVisible[static_cast<int32>(ETimerNodeType::CpuScope)]));
+	Settings.SetTimersViewShowCpuSamplingTimers(static_cast<int32>(bTimerTypeIsVisible[static_cast<int32>(ETimerNodeType::CpuSampling)]));
 	Settings.SetTimersViewShowZeroCountTimers(!bFilterOutZeroCountTimers);
 
 	Settings.SaveToConfig();
@@ -4007,9 +4040,10 @@ void STimersView::LoadSettings()
 
 	ModeFrameType = static_cast<ETraceFrameType>(Settings.GetTimersViewMode());
 	GroupingMode = static_cast<ETimerGroupingMode>(Settings.GetTimersViewGroupingMode());
-	bTimerTypeIsVisible[static_cast<int32>(ETimerNodeType::GpuScope)] = Settings.GetTimersViewShowGpuEvents();
-	bTimerTypeIsVisible[static_cast<int32>(ETimerNodeType::VerseScope)] = Settings.GetTimersViewShowVerseEvents();
-	bTimerTypeIsVisible[static_cast<int32>(ETimerNodeType::CpuScope)] = Settings.GetTimersViewShowCpuEvents();
+	bTimerTypeIsVisible[static_cast<int32>(ETimerNodeType::GpuScope)] = Settings.GetTimersViewShowGpuTimers();
+	bTimerTypeIsVisible[static_cast<int32>(ETimerNodeType::VerseSampling)] = Settings.GetTimersViewShowVerseSamplingTimers();
+	bTimerTypeIsVisible[static_cast<int32>(ETimerNodeType::CpuScope)] = Settings.GetTimersViewShowCpuTimers();
+	bTimerTypeIsVisible[static_cast<int32>(ETimerNodeType::CpuSampling)] = Settings.GetTimersViewShowCpuSamplingTimers();
 	bFilterOutZeroCountTimers = !Settings.GetTimersViewShowZeroCountTimers();
 
 	LoadVisibleColumnsSettings();
