[Insights] Insight Settings: Persistent Settings for Context Switches Toggle Options
* Persistency for Context Switches options in Timing Profiler
* Fixed incorrect DefaultUnrealInsightsSettings.ini labels

diff --git a/Engine/Programs/UnrealInsights/Config/DefaultUnrealInsightsSettings.ini b/Engine/Programs/UnrealInsights/Config/DefaultUnrealInsightsSettings.ini
--- a/Engine/Programs/UnrealInsights/Config/DefaultUnrealInsightsSettings.ini
+++ b/Engine/Programs/UnrealInsights/Config/DefaultUnrealInsightsSettings.ini
@@ -79,6 +79,12 @@ ShowAllCpuTracks=True
 ;+VisibleCpuThreadGroups=Render
 ;+HiddenCpuThreadGroups=Trace
 
+ShowCpuCoreTracks=True
+ShowContextSwitches=True
+ShowContextSwitchOverlays=True
+ShowContextSwitchExtendedLines=True
+ShowNonTargetProcessEvents=True
+
 [Insights.TimingProfiler.MainGraph]
 ; These options can be changed from the context menu of the Main Graph track (Timing Insights profiler / Timing view).
 
@@ -134,8 +140,11 @@ GroupingMode=3
 ; Controls the CPU filter button from Timers View.
 ShowCpuTimers=True
 
+; Controls the CPU Sampling filter button from Timers View.
+ShowCpuSamplingTimers=True
+
 ; Controls the Verse filter button from Timers View.
-ShowVerseTimers=True
+ShowVerseSamplingTimers=True
 
 ; Controls the GPU filter button from Timers View.
 ShowGpuTimers=True
@@ -143,7 +152,7 @@ ShowGpuTimers=True
 ; Controls the "Filter Out Zero Count Timers" button from Timers View.
 ShowZeroCountTimers=True
 
-[Insights.AssetLoadingProfiler.TimingView]
+[Insights.LoadingProfiler.TimingView]
 ; These options can be changed from the CPU/GPU drop down menus of Asset Loading Insights profiler / Timing view.
 
 ShowAllGpuTracks=False
@@ -158,6 +167,12 @@ ShowAllVerseTracks=False
 
 ShowAllCpuTracks=False
 
+ShowCpuCoreTracks=False
+ShowContextSwitches=False
+ShowContextSwitchOverlays=True
+ShowContextSwitchExtendedLines=True
+ShowNonTargetProcessEvents=True
+
 [Insights.MemoryProfiler]
 
 ; Add symbol search paths by adding to this array. Order is important.
@@ -182,3 +197,9 @@ ShowGpuFencesExtendedLines=True
 ShowAllVerseTracks=False
 
 ShowAllCpuTracks=False
+
+ShowCpuCoreTracks=False
+ShowContextSwitches=False
+ShowContextSwitchOverlays=True
+ShowContextSwitchExtendedLines=True
+ShowNonTargetProcessEvents=True
diff --git a/Engine/Source/Developer/TraceInsights/Private/Insights/ContextSwitches/ContextSwitchesProfilerManager.cpp b/Engine/Source/Developer/TraceInsights/Private/Insights/ContextSwitches/ContextSwitchesProfilerManager.cpp
--- a/Engine/Source/Developer/TraceInsights/Private/Insights/ContextSwitches/ContextSwitchesProfilerManager.cpp
+++ b/Engine/Source/Developer/TraceInsights/Private/Insights/ContextSwitches/ContextSwitchesProfilerManager.cpp
@@ -161,9 +161,9 @@ bool FContextSwitchesProfilerManager::Tick(float DeltaTime)
 
 				if (!ContextSwitchesSharedState.IsValid())
 				{
-					ContextSwitchesSharedState = MakeShared<FContextSwitchesSharedState>(TimingView.Get());
-					ContextSwitchesSharedState->AddCommands();
+					ContextSwitchesSharedState = MakeShared<FContextSwitchesSharedState>(nullptr);
 					IModularFeatures::Get().RegisterModularFeature(Timing::TimingViewExtenderFeatureName, ContextSwitchesSharedState.Get());
+					ContextSwitchesSharedState->OnBeginSession(*TimingView.Get());
 				}
 			}
 
diff --git a/Engine/Source/Developer/TraceInsights/Private/Insights/ContextSwitches/ViewModels/ContextSwitchesSharedState.cpp b/Engine/Source/Developer/TraceInsights/Private/Insights/ContextSwitches/ViewModels/ContextSwitchesSharedState.cpp
--- a/Engine/Source/Developer/TraceInsights/Private/Insights/ContextSwitches/ViewModels/ContextSwitchesSharedState.cpp
+++ b/Engine/Source/Developer/TraceInsights/Private/Insights/ContextSwitches/ViewModels/ContextSwitchesSharedState.cpp
@@ -53,7 +53,7 @@ FContextSwitchesStateCommands::FContextSwitchesStateCommands()
 UE_DISABLE_OPTIMIZATION_SHIP
 void FContextSwitchesStateCommands::RegisterCommands()
 {
-	UI_COMMAND(Command_ShowCoreTracks,
+	UI_COMMAND(Command_ShowCpuCoreTracks,
 		"CPU Core Tracks",
 		"Shows/hides the CPU Core tracks.",
 		EUserInterfaceActionType::ToggleButton,
@@ -117,12 +117,8 @@ FContextSwitchesSharedState::FContextSwitchesSharedState(TimingProfiler::STiming
 	: TimingViewSession(InTimingView)
 	, ThreadsSerial(0)
 	, CpuCoresSerial(0)
-	, bAreCoreTracksVisible(true)
-	, bAreContextSwitchesVisible(true)
-	, bAreOverlaysVisible(true)
-	, bAreExtendedLinesVisible(true)
-	, bAreNonTargetProcessEventsVisible(true)
 	, bSyncWithProviders(true)
+	, TimingViewSettings(MakeShared<FTimingViewLocalSettings>())
 {
 }
 
@@ -136,6 +132,8 @@ void FContextSwitchesSharedState::OnBeginSession(Timing::ITimingViewSession& InS
 		{
 			TimingViewSession = &InSession;
 			AddCommands();
+
+			TimingViewSettings = MakeShared<FTimingViewPersistentSettings>(InSession.GetName());
 		}
 		else
 		{
@@ -149,12 +147,6 @@ void FContextSwitchesSharedState::OnBeginSession(Timing::ITimingViewSession& InS
 	ThreadsSerial = 0;
 	CpuCoresSerial = 0;
 
-	bAreCoreTracksVisible = true;
-	bAreContextSwitchesVisible = true;
-	bAreOverlaysVisible = true;
-	bAreExtendedLinesVisible = true;
-	bAreNonTargetProcessEventsVisible = true;
-
 	bSyncWithProviders = true;
 
 	TargetTimingEvent = nullptr;
@@ -175,11 +167,6 @@ void FContextSwitchesSharedState::OnEndSession(Timing::ITimingViewSession& InSes
 	ThreadsSerial = 0;
 	CpuCoresSerial = 0;
 
-	bAreCoreTracksVisible = true;
-	bAreContextSwitchesVisible = true;
-	bAreOverlaysVisible = true;
-	bAreExtendedLinesVisible = true;
-
 	bSyncWithProviders = false;
 
 	TargetTimingEvent = nullptr;
@@ -197,12 +184,12 @@ void FContextSwitchesSharedState::Tick(Timing::ITimingViewSession& InSession, co
 
 	if (bSyncWithProviders && AreContextSwitchesAvailable())
 	{
-		if (bAreCoreTracksVisible)
+		if (AreCpuCoreTracksVisible())
 		{
-			AddCoreTracks();
+			AddCpuCoreTracks();
 		}
 
-		if (bAreContextSwitchesVisible)
+		if (AreContextSwitchesVisible())
 		{
 			AddContextSwitchesChildTracks();
 		}
@@ -245,7 +232,7 @@ void FContextSwitchesSharedState::BuildSubMenu(FMenuBuilder& InMenuBuilder)
 {
 	InMenuBuilder.BeginSection("ContextSwitches", LOCTEXT("ContextMenu_Section_ContextSwitches", "Context Switches"));
 	{
-		InMenuBuilder.AddMenuEntry(FContextSwitchesStateCommands::Get().Command_ShowCoreTracks);
+		InMenuBuilder.AddMenuEntry(FContextSwitchesStateCommands::Get().Command_ShowCpuCoreTracks);
 		InMenuBuilder.AddMenuEntry(FContextSwitchesStateCommands::Get().Command_ShowNonTargetProcessEvents);
 		InMenuBuilder.AddMenuEntry(FContextSwitchesStateCommands::Get().Command_ShowContextSwitches);
 		InMenuBuilder.AddMenuEntry(FContextSwitchesStateCommands::Get().Command_ShowOverlays);
@@ -314,10 +301,10 @@ void FContextSwitchesSharedState::AddCommands()
 	ensure(CommandList.IsValid());
 
 	CommandList->MapAction(
-		FContextSwitchesStateCommands::Get().Command_ShowCoreTracks,
-		FExecuteAction::CreateSP(this, &FContextSwitchesSharedState::Command_ShowCoreTracks_Execute),
-		FCanExecuteAction::CreateSP(this, &FContextSwitchesSharedState::Command_ShowCoreTracks_CanExecute),
-		FIsActionChecked::CreateSP(this, &FContextSwitchesSharedState::Command_ShowCoreTracks_IsChecked));
+		FContextSwitchesStateCommands::Get().Command_ShowCpuCoreTracks,
+		FExecuteAction::CreateSP(this, &FContextSwitchesSharedState::Command_ShowCpuCoreTracks_Execute),
+		FCanExecuteAction::CreateSP(this, &FContextSwitchesSharedState::Command_ShowCpuCoreTracks_CanExecute),
+		FIsActionChecked::CreateSP(this, &FContextSwitchesSharedState::Command_ShowCpuCoreTracks_IsChecked));
 
 	CommandList->MapAction(
 		FContextSwitchesStateCommands::Get().Command_ShowNonTargetProcessEvents,
@@ -373,7 +360,7 @@ bool FContextSwitchesSharedState::AreContextSwitchesAvailable() const
 
 ////////////////////////////////////////////////////////////////////////////////////////////////////
 
-void FContextSwitchesSharedState::AddCoreTracks()
+void FContextSwitchesSharedState::AddCpuCoreTracks()
 {
 	TSharedPtr<TimingProfiler::STimingView> TimingView = GetTimingView();
 
@@ -421,7 +408,7 @@ void FContextSwitchesSharedState::AddCoreTracks()
 
 ////////////////////////////////////////////////////////////////////////////////////////////////////
 
-void FContextSwitchesSharedState::RemoveCoreTracks()
+void FContextSwitchesSharedState::RemoveCpuCoreTracks()
 {
 	TSharedPtr<TimingProfiler::STimingView> TimingView = GetTimingView();
 	CpuCoresSerial = 0;
@@ -505,19 +492,19 @@ void FContextSwitchesSharedState::RemoveContextSwitchesChildTracks()
 
 ////////////////////////////////////////////////////////////////////////////////////////////////////
 
-void FContextSwitchesSharedState::SetCoreTracksVisible(bool bOnOff)
+void FContextSwitchesSharedState::SetCpuCoreTracksVisible(bool bOnOff)
 {
-	if (bAreCoreTracksVisible != bOnOff)
+	if (AreCpuCoreTracksVisible() != bOnOff)
 	{
-		bAreCoreTracksVisible = bOnOff;
+		TimingViewSettings->SetToggleOption(FTimingViewSettings::Option_ShowCpuCoreTracks, bOnOff);
 
-		if (bAreCoreTracksVisible)
+		if (bOnOff)
 		{
-			AddCoreTracks();
+			AddCpuCoreTracks();
 		}
 		else
 		{
-			RemoveCoreTracks();
+			RemoveCpuCoreTracks();
 		}
 	}
 }
@@ -526,11 +513,11 @@ void FContextSwitchesSharedState::SetCoreTracksVisible(bool bOnOff)
 
 void FContextSwitchesSharedState::SetContextSwitchesVisible(bool bOnOff)
 {
-	if (bAreContextSwitchesVisible != bOnOff)
+	if (AreContextSwitchesVisible() != bOnOff)
 	{
-		bAreContextSwitchesVisible = bOnOff;
+		TimingViewSettings->SetToggleOption(FTimingViewSettings::Option_ShowContextSwitches, bOnOff);
 
-		if (bAreContextSwitchesVisible)
+		if (bOnOff)
 		{
 			AddContextSwitchesChildTracks();
 		}
@@ -545,21 +532,21 @@ void FContextSwitchesSharedState::SetContextSwitchesVisible(bool bOnOff)
 
 void FContextSwitchesSharedState::SetOverlaysVisible(bool bOnOff)
 {
-	bAreOverlaysVisible = bOnOff;
+	TimingViewSettings->SetToggleOption(FTimingViewSettings::Option_ShowContextSwitchOverlays, bOnOff);
 }
 
 ////////////////////////////////////////////////////////////////////////////////////////////////////
 
 void FContextSwitchesSharedState::SetExtendedLinesVisible(bool bOnOff)
 {
-	bAreExtendedLinesVisible = bOnOff;
+	TimingViewSettings->SetToggleOption(FTimingViewSettings::Option_ShowContextSwitchExtendedLines, bOnOff);
 }
 
 ////////////////////////////////////////////////////////////////////////////////////////////////////
 
 void FContextSwitchesSharedState::SetNonTargetProcessEventsVisible(bool bOnOff)
 {
-	if (bAreNonTargetProcessEventsVisible != bOnOff)
+	if (AreNonTargetProcessEventsVisible() != bOnOff)
 	{
 		TSharedPtr<TimingProfiler::STimingView> TimingView = GetTimingView();
 		const TSharedPtr<ITimingEventFilter> EventFilter = TimingView->GetEventFilter();
@@ -575,7 +562,7 @@ void FContextSwitchesSharedState::SetNonTargetProcessEventsVisible(bool bOnOff)
 			}
 		}
 
-		bAreNonTargetProcessEventsVisible = bOnOff;
+		TimingViewSettings->SetToggleOption(FTimingViewSettings::Option_ShowNonTargetProcessEvents, bOnOff);
 	}
 }
 
diff --git a/Engine/Source/Developer/TraceInsights/Private/Insights/ContextSwitches/ViewModels/ContextSwitchesSharedState.h b/Engine/Source/Developer/TraceInsights/Private/Insights/ContextSwitches/ViewModels/ContextSwitchesSharedState.h
--- a/Engine/Source/Developer/TraceInsights/Private/Insights/ContextSwitches/ViewModels/ContextSwitchesSharedState.h
+++ b/Engine/Source/Developer/TraceInsights/Private/Insights/ContextSwitches/ViewModels/ContextSwitchesSharedState.h
@@ -8,6 +8,8 @@
 
 // TraceInsights
 #include "Insights/ITimingViewExtender.h"
+#include "Insights/InsightsManager.h"
+#include "Insights/InsightsSettings.h"
 
 ////////////////////////////////////////////////////////////////////////////////////////////////////
 
@@ -42,7 +44,7 @@ public:
 	virtual void RegisterCommands() override;
 
 	// Commands for the Tracks Filter menu.
-	TSharedPtr<FUICommandInfo> Command_ShowCoreTracks;
+	TSharedPtr<FUICommandInfo> Command_ShowCpuCoreTracks;
 	TSharedPtr<FUICommandInfo> Command_ShowContextSwitches;
 	TSharedPtr<FUICommandInfo> Command_ShowOverlays;
 	TSharedPtr<FUICommandInfo> Command_ShowExtendedLines;
@@ -61,6 +63,70 @@ public:
 
 class FContextSwitchesSharedState : public Timing::ITimingViewExtender, public TSharedFromThis<FContextSwitchesSharedState>
 {
+	class ITimingViewSettings
+	{
+	public:
+		virtual ~ITimingViewSettings() {}
+
+		virtual bool IsToggleOptionEnabled(const TCHAR* Key) const = 0;
+		virtual void SetToggleOption(const TCHAR* Key, bool Value) = 0;
+	};
+
+	class FTimingViewPersistentSettings : public ITimingViewSettings
+	{
+	public:
+		FTimingViewPersistentSettings(const FName& InTimingViewName)
+			: Settings(GetInsightsSettings().GetOrCreateTimingViewSettings(InTimingViewName))
+		{
+		}
+
+		virtual ~FTimingViewPersistentSettings()
+		{
+		}
+
+		const FName& GetName() const
+		{
+			return Settings.GetName();
+		}
+
+		//////////////////////////////////////////////////
+
+		virtual bool IsToggleOptionEnabled(const TCHAR* Key) const override
+		{
+			return Settings.IsToggleOptionEnabled(Key);
+		}
+		virtual void SetToggleOption(const TCHAR* Key, bool Value) override
+		{
+			GetInsightsSettings().SetAndSaveTimingViewToggleOption(GetName(), Key, Value);
+		}
+
+	private:
+		FInsightsSettings& GetInsightsSettings()
+		{
+			return UE::Insights::FInsightsManager::Get()->GetSettings();
+		}
+		const FInsightsSettings& GetInsightsSettings() const
+		{
+			return UE::Insights::FInsightsManager::Get()->GetSettings();
+		}
+
+	private:
+		FTimingViewSettings& Settings;
+	};
+	
+	class FTimingViewLocalSettings : public ITimingViewSettings
+    {
+    public:
+    	FTimingViewLocalSettings() {}
+    	virtual ~FTimingViewLocalSettings() {}
+
+    	virtual bool IsToggleOptionEnabled(const TCHAR* Key) const override              { return Settings.IsToggleOptionEnabled(Key); }
+    	virtual void SetToggleOption(const TCHAR* Key, bool Value) override              { Settings.SetToggleOption(Key, Value); }
+
+    private:
+    	FTimingViewSettings Settings;
+    };
+	
 public:
 	FContextSwitchesSharedState(TimingProfiler::STimingView* InTimingView);
 	virtual ~FContextSwitchesSharedState() = default;
@@ -80,34 +146,34 @@ public:
 
 	bool AreContextSwitchesAvailable() const;
 
-	bool AreCoreTracksVisible() const { return bAreCoreTracksVisible; }
-	void ShowCoreTracks() { SetCoreTracksVisible(true); }
-	void HideCoreTracks() { SetCoreTracksVisible(false); }
-	void ToggleCoreTracks() { SetCoreTracksVisible(!bAreCoreTracksVisible); }
-	void SetCoreTracksVisible(bool bOnOff);
+	bool AreCpuCoreTracksVisible() const { return TimingViewSettings->IsToggleOptionEnabled(FTimingViewSettings::Option_ShowCpuCoreTracks); }
+	void ShowCpuCoreTracks() { SetCpuCoreTracksVisible(true); }
+	void HideCpuCoreTracks() { SetCpuCoreTracksVisible(false); }
+	void ToggleCpuCoreTracks() { SetCpuCoreTracksVisible(!AreCpuCoreTracksVisible()); }
+	void SetCpuCoreTracksVisible(bool bOnOff);
 
-	bool AreContextSwitchesVisible() const { return bAreContextSwitchesVisible; }
+	bool AreContextSwitchesVisible() const { return TimingViewSettings->IsToggleOptionEnabled(FTimingViewSettings::Option_ShowContextSwitches); }
 	void ShowContextSwitches() { SetContextSwitchesVisible(true); }
 	void HideContextSwitches() { SetContextSwitchesVisible(false); }
-	void ToggleContextSwitches() { SetContextSwitchesVisible(!bAreContextSwitchesVisible); }
+	void ToggleContextSwitches() { SetContextSwitchesVisible(!AreContextSwitchesVisible()); }
 	void SetContextSwitchesVisible(bool bOnOff);
 
-	bool AreOverlaysVisible() const { return bAreOverlaysVisible; }
+	bool AreOverlaysVisible() const { return TimingViewSettings->IsToggleOptionEnabled(FTimingViewSettings::Option_ShowContextSwitchOverlays); }
 	void ShowOverlays() { SetOverlaysVisible(true); }
 	void HideOverlays() { SetOverlaysVisible(false); }
-	void ToggleOverlays() { SetOverlaysVisible(!bAreOverlaysVisible); }
+	void ToggleOverlays() { SetOverlaysVisible(!AreOverlaysVisible()); }
 	void SetOverlaysVisible(bool bOnOff);
 
-	bool AreExtendedLinesVisible() const { return bAreExtendedLinesVisible; }
+	bool AreExtendedLinesVisible() const { return TimingViewSettings->IsToggleOptionEnabled(FTimingViewSettings::Option_ShowContextSwitchExtendedLines); }
 	void ShowExtendedLines() { SetExtendedLinesVisible(true); }
 	void HideExtendedLines() { SetExtendedLinesVisible(false); }
-	void ToggleExtendedLines() { SetExtendedLinesVisible(!bAreExtendedLinesVisible); }
+	void ToggleExtendedLines() { SetExtendedLinesVisible(!AreExtendedLinesVisible()); }
 	void SetExtendedLinesVisible(bool bOnOff);
 
-	bool AreNonTargetProcessEventsVisible() const { return bAreNonTargetProcessEventsVisible; }
+	bool AreNonTargetProcessEventsVisible() const { return TimingViewSettings->IsToggleOptionEnabled(FTimingViewSettings::Option_ShowNonTargetProcessEvents); }
 	void ShowNonTargetProcessEvents() { SetNonTargetProcessEventsVisible(true); }
 	void HideNonTargetProcessEvents() { SetNonTargetProcessEventsVisible(false); }
-	void ToggleNonTargetProcessEvents() { SetNonTargetProcessEventsVisible(!bAreNonTargetProcessEventsVisible); }
+	void ToggleNonTargetProcessEvents() { SetNonTargetProcessEventsVisible(!AreNonTargetProcessEventsVisible()); }
 	void SetNonTargetProcessEventsVisible(bool bOnOff);
 
 	void SetTargetTimingEvent(const TSharedPtr<const ITimingEvent> InEvent)
@@ -122,17 +188,17 @@ public:
 	TSharedPtr<FCpuCoreTimingTrack> GetCpuCoreTimingTrack(uint32 CoreNumber) const;
 
 private:
-	void AddCoreTracks();
-	void RemoveCoreTracks();
+	void AddCpuCoreTracks();
+	void RemoveCpuCoreTracks();
 
 	void AddContextSwitchesChildTracks();
 	void RemoveContextSwitchesChildTracks();
 
 	void BuildSubMenu(FMenuBuilder& InMenuBuilder);
 
-	void Command_ShowCoreTracks_Execute() { ToggleCoreTracks(); }
-	bool Command_ShowCoreTracks_CanExecute() const { return AreContextSwitchesAvailable(); }
-	bool Command_ShowCoreTracks_IsChecked() const { return AreCoreTracksVisible(); }
+	void Command_ShowCpuCoreTracks_Execute() { ToggleCpuCoreTracks(); }
+	bool Command_ShowCpuCoreTracks_CanExecute() const { return AreContextSwitchesAvailable(); }
+	bool Command_ShowCpuCoreTracks_IsChecked() const { return AreCpuCoreTracksVisible(); }
 
 	void Command_ShowContextSwitches_Execute() { ToggleContextSwitches(); }
 	bool Command_ShowContextSwitches_CanExecute() const { return AreContextSwitchesAvailable(); }
@@ -147,7 +213,7 @@ private:
 	bool Command_ShowExtendedLines_IsChecked() const { return AreExtendedLinesVisible(); }
 
 	void Command_ShowNonTargetProcessEvents_Execute() { ToggleNonTargetProcessEvents(); }
-	bool Command_ShowNonTargetProcessEvents_CanExecute() const { return AreContextSwitchesAvailable() && AreCoreTracksVisible(); }
+	bool Command_ShowNonTargetProcessEvents_CanExecute() const { return AreContextSwitchesAvailable() && AreCpuCoreTracksVisible(); }
 	bool Command_ShowNonTargetProcessEvents_IsChecked() const { return AreNonTargetProcessEventsVisible(); }
 
 	bool IsValidCpuCoreEventSelected() const;
@@ -160,10 +226,10 @@ private:
 	bool Command_DockCpuThreadTrackToBottom_CanExecute() const { return AreContextSwitchesAvailable() && AreContextSwitchesVisible() && IsValidCpuCoreEventSelected(); }
 
 	void Command_NavigateToCpuCoreEvent_Execute();
-	bool Command_NavigateToCpuCoreEvent_CanExecute() const { return AreContextSwitchesAvailable() && AreCoreTracksVisible() && IsValidContextSwitchEventSelected(); }
+	bool Command_NavigateToCpuCoreEvent_CanExecute() const { return AreContextSwitchesAvailable() && AreCpuCoreTracksVisible() && IsValidContextSwitchEventSelected(); }
 
 	void Command_DockCpuCoreTrackToTop_Execute();
-	bool Command_DockCpuCoreTrackToTop_CanExecute() const { return AreContextSwitchesAvailable() && AreCoreTracksVisible() && IsValidContextSwitchEventSelected(); }
+	bool Command_DockCpuCoreTrackToTop_CanExecute() const { return AreContextSwitchesAvailable() && AreCpuCoreTracksVisible() && IsValidContextSwitchEventSelected(); }
 
 	void PopulateCoreEventNameSuggestionList(const FString& Text, TArray<FString>& OutSuggestions);
 
@@ -176,15 +242,10 @@ private:
 	uint64 ThreadsSerial;
 	uint64 CpuCoresSerial;
 
-	bool bAreCoreTracksVisible;
-	bool bAreContextSwitchesVisible;
-	bool bAreOverlaysVisible;
-	bool bAreExtendedLinesVisible;
-	bool bAreNonTargetProcessEventsVisible;
-
 	bool bSyncWithProviders;
 
 	TSharedPtr<const ITimingEvent> TargetTimingEvent;
+	TSharedRef<ITimingViewSettings> TimingViewSettings;
 };
 
 ////////////////////////////////////////////////////////////////////////////////////////////////////
diff --git a/Engine/Source/Developer/TraceInsights/Private/Insights/InsightsSettings.h b/Engine/Source/Developer/TraceInsights/Private/Insights/InsightsSettings.h
--- a/Engine/Source/Developer/TraceInsights/Private/Insights/InsightsSettings.h
+++ b/Engine/Source/Developer/TraceInsights/Private/Insights/InsightsSettings.h
@@ -159,6 +159,12 @@ private:
 		RegisterToggleOptions(Option_ShowAllVerseTracks, true);
 
 		RegisterToggleOptions(Option_ShowAllCpuTracks, true);
+
+		RegisterToggleOptions(Option_ShowCpuCoreTracks, true);
+		RegisterToggleOptions(Option_ShowContextSwitches, true);
+		RegisterToggleOptions(Option_ShowContextSwitchOverlays, true);
+		RegisterToggleOptions(Option_ShowContextSwitchExtendedLines, true);
+		RegisterToggleOptions(Option_ShowNonTargetProcessEvents, true);
 	}
 
 public:
@@ -183,6 +189,17 @@ public:
 	/** Toggle all or no CPU Tracks in the Timing View. */
 	static constexpr const TCHAR* Option_ShowAllCpuTracks = TEXT("ShowAllCpuTracks");
 
+	/** Toggle all or no CPU Core Tracks in the Timing View. */
+	static constexpr const TCHAR* Option_ShowCpuCoreTracks = TEXT("ShowCpuCoreTracks");
+	/** Visibility of the Context Switch tracks. */
+	static constexpr const TCHAR* Option_ShowContextSwitches = TEXT("ShowContextSwitches");
+	/** Extends the visualization of context work events over the CPU timing tracks. */
+	static constexpr const TCHAR* Option_ShowContextSwitchOverlays = TEXT("ShowContextSwitchOverlays");
+	/** Extended vertical lines at the edges of each context work event. */
+	static constexpr const TCHAR* Option_ShowContextSwitchExtendedLines = TEXT("ShowContextSwitchExtendedLines");
+	/** Visibility of the Non-Target events. */
+	static constexpr const TCHAR* Option_ShowNonTargetProcessEvents = TEXT("ShowNonTargetProcessEvents");
+
 private:
 	// The unique name of the Timing View settings.
 	// When settings are saved it will use a section named [Insights.(TimingViewName).TimingView].
