[Insights] Fixed crash when creating a Cpu Sampling timer, but its timer name is not yet initialized.

#rnx

diff --git a/Engine/Source/Developer/TraceInsights/Private/Insights/TimingProfiler/ViewModels/TimerNode.cpp b/Engine/Source/Developer/TraceInsights/Private/Insights/TimingProfiler/ViewModels/TimerNode.cpp
--- a/Engine/Source/Developer/TraceInsights/Private/Insights/TimingProfiler/ViewModels/TimerNode.cpp
+++ b/Engine/Source/Developer/TraceInsights/Private/Insights/TimingProfiler/ViewModels/TimerNode.cpp
@@ -47,6 +47,11 @@ static FName GetCpuSamplingGroup()
 
 FName MakeSafeFName(const TCHAR* InName)
 {
+	if (!InName || *InName == TEXT('\0'))
+	{
+		return FName();
+	}
+
 	int32 Len = FCString::Strlen(InName);
 	if (Len >= NAME_SIZE)
 	{
