[Insights] Insight Settings: Persistent Setting for CPU Stack Sampling Toggle Option
* Toggle option to show all Stack Sample Tracks
* Prevent Verse / Sampling Tracks from being toggled unless they have active tracks

diff --git a/Engine/Programs/UnrealInsights/Config/DefaultUnrealInsightsSettings.ini b/Engine/Programs/UnrealInsights/Config/DefaultUnrealInsightsSettings.ini
--- a/Engine/Programs/UnrealInsights/Config/DefaultUnrealInsightsSettings.ini
+++ b/Engine/Programs/UnrealInsights/Config/DefaultUnrealInsightsSettings.ini
@@ -79,6 +79,8 @@ ShowAllCpuTracks=True
 ;+VisibleCpuThreadGroups=Render
 ;+HiddenCpuThreadGroups=Trace
 
+ShowAllCpuSamplingTracks=True
+
 ShowCpuCoreTracks=True
 ShowContextSwitches=True
 ShowContextSwitchOverlays=True
@@ -167,6 +169,8 @@ ShowAllVerseTracks=False
 
 ShowAllCpuTracks=False
 
+ShowAllCpuSamplingTracks=False
+
 ShowCpuCoreTracks=False
 ShowContextSwitches=False
 ShowContextSwitchOverlays=True
@@ -198,6 +202,8 @@ ShowAllVerseTracks=False
 
 ShowAllCpuTracks=False
 
+ShowAllCpuSamplingTracks=False
+
 ShowCpuCoreTracks=False
 ShowContextSwitches=False
 ShowContextSwitchOverlays=True
diff --git a/Engine/Source/Developer/TraceInsights/Private/Insights/InsightsSettings.h b/Engine/Source/Developer/TraceInsights/Private/Insights/InsightsSettings.h
--- a/Engine/Source/Developer/TraceInsights/Private/Insights/InsightsSettings.h
+++ b/Engine/Source/Developer/TraceInsights/Private/Insights/InsightsSettings.h
@@ -160,6 +160,8 @@ private:
 
 		RegisterToggleOptions(Option_ShowAllCpuTracks, true);
 
+		RegisterToggleOptions(Option_ShowAllCpuSamplingTracks, true);
+
 		RegisterToggleOptions(Option_ShowCpuCoreTracks, true);
 		RegisterToggleOptions(Option_ShowContextSwitches, true);
 		RegisterToggleOptions(Option_ShowContextSwitchOverlays, true);
@@ -189,6 +191,9 @@ public:
 	/** Toggle all or no CPU Tracks in the Timing View. */
 	static constexpr const TCHAR* Option_ShowAllCpuTracks = TEXT("ShowAllCpuTracks");
 
+	/** Toggle all or no CPU Sampling Tracks in the Timing View. */
+	static constexpr const TCHAR* Option_ShowAllCpuSamplingTracks = TEXT("ShowAllCpuSamplingTracks");
+
 	/** Toggle all or no CPU Core Tracks in the Timing View. */
 	static constexpr const TCHAR* Option_ShowCpuCoreTracks = TEXT("ShowCpuCoreTracks");
 	/** Visibility of the Context Switch tracks. */
diff --git a/Engine/Source/Developer/TraceInsights/Private/Insights/TimingProfiler/ViewModels/ThreadTimingSharedState.cpp b/Engine/Source/Developer/TraceInsights/Private/Insights/TimingProfiler/ViewModels/ThreadTimingSharedState.cpp
--- a/Engine/Source/Developer/TraceInsights/Private/Insights/TimingProfiler/ViewModels/ThreadTimingSharedState.cpp
+++ b/Engine/Source/Developer/TraceInsights/Private/Insights/TimingProfiler/ViewModels/ThreadTimingSharedState.cpp
@@ -75,6 +75,9 @@ public:
 
 	/** Toggles visibility for all CPU thread tracks at once. */
 	TSharedPtr<FUICommandInfo> ShowHideAllCpuTracks;
+
+	/** Toggles visibility for all CPU Sampling thread tracks at once. */
+	TSharedPtr<FUICommandInfo> ShowHideAllCpuSamplingTracks;
 };
 
 ////////////////////////////////////////////////////////////////////////////////////////////////////
@@ -153,6 +156,12 @@ void FThreadTimingViewCommands::RegisterCommands()
 		"Shows/hides all CPU tracks (and all CPU thread groups).",
 		EUserInterfaceActionType::ToggleButton,
 		FInputChord(EKeys::U));
+
+	UI_COMMAND(ShowHideAllCpuSamplingTracks,
+		"CPU Stack Sampling Tracks",
+		"Shows/hides all CPU Stack Sampling thread tracks.",
+		EUserInterfaceActionType::ToggleButton,
+		FInputChord());
 }
 UE_ENABLE_OPTIMIZATION_SHIP
 
@@ -316,12 +325,14 @@ void FThreadTimingSharedState::ResetToDefaultValues()
 		ShowAllGpuTracks();
 		ShowAllVerseTracks();
 		ShowAllCpuTracks();
+		ShowAllCpuSamplingTracks();
 	}
 	else
 	{
 		HideAllGpuTracks();
 		HideAllVerseTracks();
 		HideAllCpuTracks();
+		HideAllCpuSamplingTracks();
 	}
 
 	TimingViewSettings->SetToggleOption(FTimingViewSettings::Option_ShowGpuWorkOverlays, true);
@@ -332,6 +343,30 @@ void FThreadTimingSharedState::ResetToDefaultValues()
 	TimingViewSettings->SetToggleOption(FTimingViewSettings::Option_ShowGpuFencesExtendedLines, true);
 	SetGpuFencesTracksVisibility(true);
 }
+	
+////////////////////////////////////////////////////////////////////////////////////////////////////
+
+bool FThreadTimingSharedState::AreCpuSamplingTracksAvailable() const
+{
+	return CpuStackSamplingTimelineCount > 0;
+}
+
+////////////////////////////////////////////////////////////////////////////////////////////////////
+
+bool FThreadTimingSharedState::AreVerseTracksAvailable() const
+{
+	TSharedPtr<const TraceServices::IAnalysisSession> Session = FInsightsManager::Get()->GetSession();
+	if (Session.IsValid())
+	{
+		if (const TraceServices::ITimingProfilerProvider* TimingProfilerProvider = TraceServices::ReadTimingProfilerProvider(*Session.Get()))
+		{
+			TraceServices::FAnalysisSessionReadScope SessionReadScope(*Session.Get());
+			return TimingProfilerProvider->HasVerseTiming();
+		}
+	}
+
+	return false;
+}
 
 ////////////////////////////////////////////////////////////////////////////////////////////////////
 
@@ -354,6 +389,7 @@ void FThreadTimingSharedState::OnBeginSession(Timing::ITimingViewSession& InSess
 		TimingViewSettings = MakeShared<FTimingViewLocalSettings>();
 
 		TimingViewSettings->SetToggleOption(FTimingViewSettings::Option_ShowAllCpuTracks, false);
+		TimingViewSettings->SetToggleOption(FTimingViewSettings::Option_ShowAllCpuSamplingTracks, false);
 		TimingViewSettings->SetToggleOption(FTimingViewSettings::Option_ShowAllGpuTracks, false);
 		TimingViewSettings->SetToggleOption(FTimingViewSettings::Option_ShowAllVerseTracks, false);
 	}
@@ -619,6 +655,7 @@ void FThreadTimingSharedState::Tick(Timing::ITimingViewSession& InSession, const
 				[this, &InSession, &bTracksOrderChanged, &CpuStackSampleTrackOrder]
 				(const TraceServices::FStackSampleTimeline& Timeline, const TraceServices::FStackSampleThread& Thread)
 				{
+					const bool bIsAllCpuSamplingTracksEnabled = IsAllCpuSamplingTracksToggleOn();
 					TSharedPtr<FCpuStackSampleTimingTrack>* TrackPtrPtr = CpuStackSampleTracks.Find(Thread.SystemThreadId);
 					if (TrackPtrPtr == nullptr)
 					{
@@ -626,12 +663,14 @@ void FThreadTimingSharedState::Tick(Timing::ITimingViewSession& InSession, const
 						TSharedPtr<FCpuStackSampleTimingTrack> Track = MakeShared<FCpuStackSampleTimingTrack>(*this, TrackName, Thread.SystemThreadId);
 						Track->SetOrder(CpuStackSampleTrackOrder);
 						CpuStackSampleTracks.Add(Thread.SystemThreadId, Track);
-						Track->SetVisibilityFlag(true);
+						Track->SetVisibilityFlag(bIsAllCpuSamplingTracksEnabled);
 						InSession.AddScrollableTrack(Track);
 					}
 					else
 					{
 						TSharedPtr<FCpuStackSampleTimingTrack> Track = *TrackPtrPtr;
+						Track->SetVisibilityFlag(bIsAllCpuSamplingTracksEnabled);
+
 						if (Track->GetOrder() != CpuStackSampleTrackOrder)
 						{
 							Track->SetOrder(CpuStackSampleTrackOrder);
@@ -683,6 +722,12 @@ void FThreadTimingSharedState::ExtendCpuTracksFilterMenu(Timing::ITimingViewSess
 		return;
 	}
 
+	InOutMenuBuilder.BeginSection("CpuSamplingTracks", LOCTEXT("ContextMenu_Section_CpuSamplingTracks", "CPU Sampling Tracks"));
+	{
+		InOutMenuBuilder.AddMenuEntry(FThreadTimingViewCommands::Get().ShowHideAllCpuSamplingTracks);
+	}
+	InOutMenuBuilder.EndSection();
+	
 	InOutMenuBuilder.BeginSection("VerseTracks", LOCTEXT("ContextMenu_Section_VerseTracks", "Verse Tracks"));
 	{
 		InOutMenuBuilder.AddMenuEntry(FThreadTimingViewCommands::Get().ShowHideAllVerseTracks);
@@ -769,7 +814,7 @@ void FThreadTimingSharedState::BindCommands()
 	CommandList->MapAction(
 		FThreadTimingViewCommands::Get().ShowHideAllVerseTracks,
 		FExecuteAction::CreateSP(this, &FThreadTimingSharedState::ShowHideAllVerseTracks),
-		FCanExecuteAction(),
+		FCanExecuteAction::CreateSP(this, &FThreadTimingSharedState::AreVerseTracksAvailable),
 		FIsActionChecked::CreateSP(this, &FThreadTimingSharedState::IsAllVerseTracksToggleOn));
 
 	CommandList->MapAction(
@@ -777,6 +822,12 @@ void FThreadTimingSharedState::BindCommands()
 		FExecuteAction::CreateSP(this, &FThreadTimingSharedState::ShowHideAllCpuTracks),
 		FCanExecuteAction(),
 		FIsActionChecked::CreateSP(this, &FThreadTimingSharedState::IsAllCpuTracksToggleOn));
+
+	CommandList->MapAction(
+		FThreadTimingViewCommands::Get().ShowHideAllCpuSamplingTracks,
+		FExecuteAction::CreateSP(this, &FThreadTimingSharedState::ShowHideAllCpuSamplingTracks),
+		FCanExecuteAction::CreateSP(this, &FThreadTimingSharedState::AreCpuSamplingTracksAvailable),
+		FIsActionChecked::CreateSP(this, &FThreadTimingSharedState::IsAllCpuSamplingTracksToggleOn));
 }
 
 ////////////////////////////////////////////////////////////////////////////////////////////////////
@@ -848,6 +899,25 @@ void FThreadTimingSharedState::SetAllCpuTracksToggle(bool bOnOff)
 		TimingView->HandleTrackVisibilityChanged();
 	}
 }
+	
+////////////////////////////////////////////////////////////////////////////////////////////////////
+
+void FThreadTimingSharedState::SetAllCpuSamplingTracksToggle(bool bOnOff)
+{
+	TimingViewSettings->SetToggleOption(FTimingViewSettings::Option_ShowAllCpuSamplingTracks, bOnOff);
+
+	for (const auto& KV : CpuStackSampleTracks)
+	{
+		FCpuStackSampleTimingTrack& Track = *KV.Value;
+		Track.SetVisibilityFlag(bOnOff);
+	}
+
+	if (TimingView)
+	{
+		TimingView->HandleTrackVisibilityChanged();
+	}
+
+}
 
 ////////////////////////////////////////////////////////////////////////////////////////////////////
 
diff --git a/Engine/Source/Developer/TraceInsights/Private/Insights/TimingProfiler/ViewModels/ThreadTimingSharedState.h b/Engine/Source/Developer/TraceInsights/Private/Insights/TimingProfiler/ViewModels/ThreadTimingSharedState.h
--- a/Engine/Source/Developer/TraceInsights/Private/Insights/TimingProfiler/ViewModels/ThreadTimingSharedState.h
+++ b/Engine/Source/Developer/TraceInsights/Private/Insights/TimingProfiler/ViewModels/ThreadTimingSharedState.h
@@ -189,6 +189,9 @@ public:
 
 	void ResetToDefaultValues();
 
+	bool AreCpuSamplingTracksAvailable() const;
+	bool AreVerseTracksAvailable() const;
+
 	bool IsAllGpuTracksToggleOn() const { return TimingViewSettings->IsToggleOptionEnabled(FTimingViewSettings::Option_ShowAllGpuTracks); }
 	void SetAllGpuTracksToggle(bool bOnOff);
 	void ShowAllGpuTracks() { SetAllGpuTracksToggle(true); }
@@ -207,6 +210,12 @@ public:
 	void HideAllCpuTracks() { SetAllCpuTracksToggle(false); }
 	void ShowHideAllCpuTracks() { SetAllCpuTracksToggle(!IsAllCpuTracksToggleOn()); }
 
+	bool IsAllCpuSamplingTracksToggleOn() const { return TimingViewSettings->IsToggleOptionEnabled(FTimingViewSettings::Option_ShowAllCpuSamplingTracks); }
+	void SetAllCpuSamplingTracksToggle(bool bOnOff);
+	void ShowAllCpuSamplingTracks() { SetAllCpuSamplingTracksToggle(true); }
+	void HideAllCpuSamplingTracks() { SetAllCpuSamplingTracksToggle(false); }
+	void ShowHideAllCpuSamplingTracks() { SetAllCpuSamplingTracksToggle(!IsAllCpuSamplingTracksToggleOn()); }
+
 	bool AreOverlaysVisibleInGpuQueueTracks() const { return TimingViewSettings->IsToggleOptionEnabled(FTimingViewSettings::Option_ShowGpuWorkOverlays); }
 	bool AreExtendedLinesVisibleInGpuQueueTracks() const { return TimingViewSettings->IsToggleOptionEnabled(FTimingViewSettings::Option_ShowGpuWorkExtendedLines); }
 
