[Insights] Timer Context Menu: Added missing Min/Max search logic for Stack Sample Events

diff --git a/Engine/Source/Developer/TraceInsights/Private/Insights/TimingProfiler/ViewModels/ThreadTimingSharedState.cpp b/Engine/Source/Developer/TraceInsights/Private/Insights/TimingProfiler/ViewModels/ThreadTimingSharedState.cpp
--- a/Engine/Source/Developer/TraceInsights/Private/Insights/TimingProfiler/ViewModels/ThreadTimingSharedState.cpp
+++ b/Engine/Source/Developer/TraceInsights/Private/Insights/TimingProfiler/ViewModels/ThreadTimingSharedState.cpp
@@ -1032,6 +1032,16 @@ TSharedPtr<const ITimingEvent> FThreadTimingSharedState::FindMaxEventInstance(ui
 		}
 	}
 
+	for (const auto& KV : CpuStackSampleTracks)
+	{
+		const FCpuStackSampleTimingTrack& Track = *KV.Value;
+		if (Track.IsVisible())
+		{
+			TrackEvent = Track.FindMaxEventInstance(TimerId, StartTime, EndTime);
+			CompareAndAssignEvent(TimingEvent, TrackEvent);
+		}
+	}
+
 	if (OldGpu1Track.IsValid() && OldGpu1Track->IsVisible())
 	{
 		TrackEvent = OldGpu1Track->FindMaxEventInstance(TimerId, StartTime, EndTime);
@@ -1087,6 +1097,16 @@ TSharedPtr<const ITimingEvent> FThreadTimingSharedState::FindMinEventInstance(ui
 		}
 	}
 
+	for (const auto& KV : CpuStackSampleTracks)
+	{
+		const FCpuStackSampleTimingTrack& Track = *KV.Value;
+		if (Track.IsVisible())
+		{
+			TrackEvent = Track.FindMinEventInstance(TimerId, StartTime, EndTime);
+			CompareAndAssignEvent(TimingEvent, TrackEvent);
+		}
+	}
+
 	if (OldGpu1Track.IsValid() && OldGpu1Track->IsVisible())
 	{
 		TrackEvent = OldGpu1Track->FindMinEventInstance(TimerId, StartTime, EndTime);
