Refactoring - InsightsTests (second swarm)

#rnx

diff --git a/Engine/Source/Developer/TraceInsights/Private/Insights/Tests/InsightsTestUtils.cpp b/Engine/Source/Developer/TraceInsights/Private/Insights/Tests/InsightsTestUtils.cpp
--- a/Engine/Source/Developer/TraceInsights/Private/Insights/Tests/InsightsTestUtils.cpp
+++ b/Engine/Source/Developer/TraceInsights/Private/Insights/Tests/InsightsTestUtils.cpp
@@ -123,22 +123,21 @@ bool FInsightsTestUtils::FileContainsString(const FString& PathToFile, const FSt
 
 bool FInsightsTestUtils::IsUnrealTraceServerReady(const TCHAR* Host, int32 Port) const
 {
-	UE::Trace::FStoreClient* StoreClient = UE::Trace::FStoreClient::Connect(Host, Port);
-	if (!StoreClient)
+	using namespace UE::Trace;
+	TSharedPtr<FStoreClient> StoreClient(FStoreClient::Connect(Host, Port));
+	if (!StoreClient.IsValid())
 	{
-		Test->AddInfo(TEXT("Cannot connect to UTS. Trying again"));
+		Test->AddInfo(TEXT("Cannot connect to UTS"));
 		return false;
 	}
 	const UE::Trace::FStoreClient::FVersion* Version = StoreClient->GetVersion();
 	if (Version == nullptr)
 	{
 		Test->AddError(TEXT("Cannot get version of UTS"));
-		delete StoreClient;
 		return false;
 	}
 	uint32 MajorVersion = Version->GetMajorVersion();
 	uint32 MinorVersion = Version->GetMinorVersion();
-	delete StoreClient;
 	Test->AddInfo(FString::Printf(TEXT("Connected to UTS version %u.%u"), MajorVersion, MinorVersion));
 	return true;
 }
@@ -150,23 +149,28 @@ bool FInsightsTestUtils::StartTracing(FTraceAuxiliary::EConnectionType Connectio
 	bool bStarted = false;
 
 	double TraceVerifyStartTime = FPlatformTime::Seconds();
-	while(FPlatformTime::Seconds() - TraceVerifyStartTime < Timeout)
+	while (FPlatformTime::Seconds() - TraceVerifyStartTime < Timeout)
 	{
-		if (ConnectionType == FTraceAuxiliary::EConnectionType::Network)
+		if (FTraceAuxiliary::IsConnected())
 		{
-			bStarted = FTraceAuxiliary::Start(FTraceAuxiliary::EConnectionType::Network, TEXT("localhost"), nullptr);
+			break;
 		}
-		else if (ConnectionType == FTraceAuxiliary::EConnectionType::File)
+
+		switch (ConnectionType)
 		{
+		case FTraceAuxiliary::EConnectionType::Network:
+			Test->AddInfo(TEXT("Connection type is Network"));
+			bStarted = FTraceAuxiliary::Start(FTraceAuxiliary::EConnectionType::Network, TEXT("localhost"), nullptr);
+			break;
+		case FTraceAuxiliary::EConnectionType::File:
+			Test->AddInfo(TEXT("Connection type is File"));
 			bStarted = FTraceAuxiliary::Start(FTraceAuxiliary::EConnectionType::File, nullptr, nullptr);
+			break;
+		default:
+			Test->AddWarning(TEXT("Not expected connection type"));
 		}
 
-		if (FTraceAuxiliary::IsConnected())
-		{
-			FPlatformProcess::Sleep(0.5f);
-			return bStarted;
-		}
-		FPlatformProcess::Sleep(0.1f);
+		FPlatformProcess::Sleep(0.4f);
 	}
 
 	return bStarted;
@@ -297,22 +301,26 @@ void FInsightsTestUtils::ResetSession() const
 }
 
 ////////////////////////////////////////////////////////////////////////////////////////////////////
-
-bool FInsightsTestUtils::IsTraceHasLiveStatus(const FString& TraceName, const TCHAR* Host, int32 Port) const
+TSharedPtr<UE::Trace::FStoreClient> FInsightsTestUtils::CreateStoreClient(const TCHAR* Host, int32 Port) const
 {
-	UE::Trace::FStoreClient* StoreClient = UE::Trace::FStoreClient::Connect(Host, Port);
-	if (!StoreClient)
+	TSharedPtr<UE::Trace::FStoreClient> StoreClient(UE::Trace::FStoreClient::Connect(Host, Port));
+	if (!StoreClient.IsValid())
 	{
 		Test->AddInfo(TEXT("The StoreClient shouldn't be null"));
-		return false;
 	}
+	return StoreClient;
+}
+
+uint32 FInsightsTestUtils::GetValidSessionCount(TSharedPtr<UE::Trace::FStoreClient> StoreClient) const
+{
 	uint32 SessionCount = StoreClient->GetSessionCount();
-	if (!SessionCount)
-	{
-		Test->AddInfo(TEXT("The SessionCount shouldn't be 0"));
-		delete StoreClient;
-		return false;
-	}
+	Test->AddInfo(FString::Printf(TEXT("The SessionCount is %d"), SessionCount));
+	return SessionCount;
+}
+
+const UE::Trace::FStoreClient::FTraceInfo* FInsightsTestUtils::FindTraceInfoByName(TSharedPtr<UE::Trace::FStoreClient> StoreClient, const FString& TraceName) const
+{
+	uint32 SessionCount = GetValidSessionCount(StoreClient);
 
 	for (uint32 Index = 0; Index < SessionCount; ++Index)
 	{
@@ -321,25 +329,40 @@ bool FInsightsTestUtils::IsTraceHasLiveStatus(const FString& TraceName, const TC
 		{
 			continue;
 		}
+
 		uint32 TraceId = SessionInfo->GetTraceId();
-		const UE::Trace::FStoreClient::FTraceInfo* Info = StoreClient->GetTraceInfoById(TraceId);
-		if (!Info)
-		{
-			continue;
-		}
-		const FUtf8StringView Utf8TraceNameView = Info->GetName();
-		FString ActualTraceName(Utf8TraceNameView);
-		if (TraceName.Contains(ActualTraceName))
+		const UE::Trace::FStoreClient::FTraceInfo* TraceInfo = StoreClient->GetTraceInfoById(TraceId);
+		if (TraceInfo)
 		{
-			Test->AddInfo(TEXT("Trace is live"));
-			delete StoreClient;
-			return true;
+			const FUtf8StringView Utf8TraceNameView = TraceInfo->GetName();
+			FString ActualTraceName(Utf8TraceNameView);
+			if (TraceName.Contains(ActualTraceName))
+			{
+				Test->AddInfo(FString::Printf(TEXT("The live trace is %s"), *ActualTraceName));
+				return TraceInfo;
+			}
 		}
-		Test->AddInfo(FString::Printf(TEXT("The live trace is %s"), *ActualTraceName));
 	}
 
-	Test->AddInfo(FString::Printf(TEXT("The trace with name %s does not have live status. Trying to find live trace"), *TraceName));
-	delete StoreClient;
+	Test->AddInfo(FString::Printf(TEXT("The trace with name %s does not have live status"), *TraceName));
+	return nullptr;
+}
+
+bool FInsightsTestUtils::IsTraceHasLiveStatus(const FString& TraceName, const TCHAR* Host, int32 Port) const
+{
+	TSharedPtr<UE::Trace::FStoreClient> StoreClient = CreateStoreClient(Host, Port);
+	if (!StoreClient.IsValid())
+	{
+		return false;
+	}
+
+	const UE::Trace::FStoreClient::FTraceInfo* TraceInfo = FindTraceInfoByName(StoreClient, TraceName);
+	if (TraceInfo)
+	{
+		Test->AddInfo(TEXT("Trace is live"));
+		return true;
+	}
+
 	return false;
 }
 
@@ -347,20 +370,13 @@ bool FInsightsTestUtils::IsTraceHasLiveStatus(const FString& TraceName, const TC
 
 FString FInsightsTestUtils::GetLiveTrace(const TCHAR* Host, int32 Port) const
 {
-	UE::Trace::FStoreClient* StoreClient = UE::Trace::FStoreClient::Connect(Host, Port);
-	if (!StoreClient)
-	{
-		Test->AddInfo(TEXT("The StoreClient shouldn't be null"));
-		return TEXT("");
-	}
-	uint32 SessionCount = StoreClient->GetSessionCount();
-	if (!SessionCount)
+	TSharedPtr<UE::Trace::FStoreClient> StoreClient = CreateStoreClient(Host, Port);
+	if (!StoreClient.IsValid())
 	{
-		Test->AddInfo(TEXT("The SessionCount shouldn't be 0"));
-		delete StoreClient;
 		return TEXT("");
 	}
 
+	const uint32 SessionCount = GetValidSessionCount(StoreClient);
 	for (uint32 Index = 0; Index < SessionCount; ++Index)
 	{
 		const UE::Trace::FStoreClient::FSessionInfo* SessionInfo = StoreClient->GetSessionInfo(Index);
@@ -371,45 +387,44 @@ FString FInsightsTestUtils::GetLiveTrace(const TCHAR* Host, int32 Port) const
 
 		uint32 TraceId = SessionInfo->GetTraceId();
 		const UE::Trace::FStoreClient::FTraceInfo* Info = StoreClient->GetTraceInfoById(TraceId);
-		FString LiveTraceName = static_cast<FString>(Info->GetName());
-		LiveTraceName = LiveTraceName + TEXT(".utrace");
-		delete StoreClient;
-		return LiveTraceName;
+		if (Info)
+		{
+			FString LiveTraceName(static_cast<FString>(Info->GetName()));
+			LiveTraceName = LiveTraceName + TEXT(".utrace");
+			return LiveTraceName;
+		}
 	}
 
 	Test->AddInfo(TEXT("There isn't any live trace"));
-	delete StoreClient;
 	return TEXT("");
 }
 
 FString FInsightsTestUtils::FInsightsTestUtils::GetUTSPath() const
 {
-	#if PLATFORM_WINDOWS
-		const FString EnginePathToUTS = FString(TEXT("Engine/Binaries/Win64/UnrealTraceServer.exe"));
-	#endif
+#if PLATFORM_WINDOWS
+	const FString EnginePathToUTS = FString(TEXT("Engine/Binaries/Win64/UnrealTraceServer.exe"));
+#endif
 
-	#if PLATFORM_MAC
-		const FString EnginePathToUTS = FString(TEXT("Engine/Binaries/Mac/UnrealTraceServer"));
-	#endif
+#if PLATFORM_MAC
+	const FString EnginePathToUTS = FString(TEXT("Engine/Binaries/Mac/UnrealTraceServer"));
+#endif
 
-	#if PLATFORM_LINUX
-		const FString EnginePathToUTS = FString(TEXT("Engine/Binaries/Linux/UnrealTraceServer"));
-	#endif
+#if PLATFORM_LINUX
+	const FString EnginePathToUTS = FString(TEXT("Engine/Binaries/Linux/UnrealTraceServer"));
+#endif
 
+	FString EntireUTSPath;
 	FString RootDirectory = FPaths::RootDir();
-	FString EntireUTSPath = FPaths::Combine(*RootDirectory, EnginePathToUTS);
-	while (FPaths::DirectoryExists(RootDirectory) && !FPaths::FileExists(EntireUTSPath))
+	while (FPaths::DirectoryExists(RootDirectory))
 	{
 		EntireUTSPath = FPaths::Combine(*RootDirectory, EnginePathToUTS);
-		if (!FPaths::FileExists(EntireUTSPath))
+		if (FPaths::FileExists(EntireUTSPath))
 		{
-			RootDirectory = FPaths::GetPath(RootDirectory);
+			return EntireUTSPath;
 		}
-	}
-	if (!FPaths::DirectoryExists(RootDirectory)) {
-		Test->AddError(TEXT("Coudln't find UTS file"));
-		EntireUTSPath = FString();
+		RootDirectory = FPaths::GetPath(RootDirectory);
 	}
 
-	return EntireUTSPath;
+	Test->AddError(TEXT("Coudln't find UTS file"));
+	return {};
 }
\ No newline at end of file
diff --git a/Engine/Source/Developer/TraceInsights/Public/Insights/Tests/InsightsTestUtils.h b/Engine/Source/Developer/TraceInsights/Public/Insights/Tests/InsightsTestUtils.h
--- a/Engine/Source/Developer/TraceInsights/Public/Insights/Tests/InsightsTestUtils.h
+++ b/Engine/Source/Developer/TraceInsights/Public/Insights/Tests/InsightsTestUtils.h
@@ -3,7 +3,7 @@
 #pragma once
 
 #include "CoreTypes.h"
-
+#include "Trace/StoreClient.h"
 #include "Containers/UnrealString.h"
 #include "ProfilingDebugging/TraceAuxiliary.h" // for FTraceAuxiliary::EConnectionType
 
@@ -28,6 +28,10 @@ public:
 	UE_API FString GetUTSPath() const;
 
 private:
+	TSharedPtr<UE::Trace::FStoreClient> CreateStoreClient(const TCHAR* Host, int32 Port) const;
+	uint32 GetValidSessionCount(TSharedPtr<UE::Trace::FStoreClient> StoreClient) const;
+	const UE::Trace::FStoreClient::FTraceInfo* FindTraceInfoByName(TSharedPtr<UE::Trace::FStoreClient> StoreClient, const FString& TraceName) const;
+
 	FAutomationTestBase* Test;
 };
 
