Add trace session status polling to Asset and Verse Insights.

diff --git a/Engine/Source/Developer/TraceInsights/Private/Insights/Widgets/STimingView.cpp b/Engine/Source/Developer/TraceInsights/Private/Insights/Widgets/STimingView.cpp
--- a/Engine/Source/Developer/TraceInsights/Private/Insights/Widgets/STimingView.cpp
+++ b/Engine/Source/Developer/TraceInsights/Private/Insights/Widgets/STimingView.cpp
@@ -809,12 +809,20 @@ void STimingView::Tick(const FGeometry& AllottedGeometry, const double InCurrent
 
 	// Tick plugin extenders.
 	// Each extender can add/remove tracks and/or change order of tracks.
+	Timing::FTimingViewExtenderTickParams TickParams(*this, Session.Get());
+	TickParams.DeltaTime = InDeltaTime;
+	TickParams.CurrentTime = InCurrentTime;
+
 	for (Timing::ITimingViewExtender* Extender : GetExtenders())
 	{
-		if (Session || Extender->ShouldAlwaysTick())
+		if (Session)
 		{
+PRAGMA_DISABLE_DEPRECATION_WARNINGS
 			Extender->Tick(*this, *Session.Get());
+PRAGMA_ENABLE_DEPRECATION_WARNINGS
 		}
+		// New Tick version.
+		Extender->Tick(TickParams);
 	}
 
 	if (Session)
diff --git a/Engine/Source/Developer/TraceInsights/Public/Insights/ITimingViewExtender.h b/Engine/Source/Developer/TraceInsights/Public/Insights/ITimingViewExtender.h
--- a/Engine/Source/Developer/TraceInsights/Public/Insights/ITimingViewExtender.h
+++ b/Engine/Source/Developer/TraceInsights/Public/Insights/ITimingViewExtender.h
@@ -25,6 +25,19 @@ class ITimingViewSession;
 
 extern TRACEINSIGHTS_API const FName TimingViewExtenderFeatureName;
 
+struct FTimingViewExtenderTickParams
+{
+	FTimingViewExtenderTickParams(ITimingViewSession& InSession, const TraceServices::IAnalysisSession* InAnalysisSession)
+		: Session(InSession)
+		, AnalysisSession(InAnalysisSession)
+	{}
+
+	ITimingViewSession& Session;
+	const TraceServices::IAnalysisSession* AnalysisSession;
+	double CurrentTime;
+	float DeltaTime;
+};
+
 class ITimingViewExtender : public IModularFeature
 {
 public:
@@ -37,7 +50,11 @@ public:
 	virtual void OnEndSession(ITimingViewSession& InSession) = 0;
 
 	/** Called each frame. If any new tracks are created they can be added via ITimingViewSession::Add*Track() */
-	virtual void Tick(ITimingViewSession& InSession, const TraceServices::IAnalysisSession& InAnalysisSession) = 0;
+	UE_DEPRECATED(5.7, "Use newer Tick overload taking TimingViewExtenderTickParams. Be advised that one can be called with a null AnalysisSession.")
+	virtual void Tick(ITimingViewSession& InSession, const TraceServices::IAnalysisSession& InAnalysisSession) {};
+
+	/** Called each frame. If any new tracks are created they can be added via ITimingViewSession::Add*Track() */
+	virtual void Tick(const FTimingViewExtenderTickParams& Params) {};
 
 	/** Extension hook for the 'CPU Tracks Filter' menu */
 	virtual void ExtendCpuTracksFilterMenu(ITimingViewSession& InSession, FMenuBuilder& InMenuBuilder) {}
@@ -57,9 +74,6 @@ public:
 
 	/** Allows extender to add functionality to the menu toolbars. */
 	virtual void ExtendMenuToolbars(ITimingViewSession& InSession, FToolBarBuilder& LeftToolbar, FToolBarBuilder& RightToolbar) {}
-
-	/** Allows extender to specify if they should receive calls to Tick() even if session is not loaded. */
-	virtual bool ShouldAlwaysTick() { return false; }
 };
 
 } // namespace UE::Insights::Timing
