Fix Verse Insights crash.

diff --git a/Engine/Source/Developer/TraceInsights/Private/Insights/TimingProfiler/Tracks/ThreadTimingTrack.cpp b/Engine/Source/Developer/TraceInsights/Private/Insights/TimingProfiler/Tracks/ThreadTimingTrack.cpp
--- a/Engine/Source/Developer/TraceInsights/Private/Insights/TimingProfiler/Tracks/ThreadTimingTrack.cpp
+++ b/Engine/Source/Developer/TraceInsights/Private/Insights/TimingProfiler/Tracks/ThreadTimingTrack.cpp
@@ -106,6 +106,20 @@ static void AppendMetadataToTooltip(FTooltipDrawState& Tooltip, const TraceServi
 		switch (Context.MajorType())
 		{
 		case ECborCode::Int:
+			{
+				int64 Value = Context.AsInt();
+				FString ValueStr;
+				if (Value > 999'999'999LL)
+				{
+					ValueStr += FString::Printf(TEXT("0x%llX"), Value);
+				}
+				else
+				{
+					ValueStr += FString::Printf(TEXT("%lld"), Value);
+				}
+				Tooltip.AddNameValueTextLine(Key, ValueStr);
+				continue;
+			}
 		case ECborCode::Uint:
 			{
 				uint64 Value = Context.AsUInt();
@@ -217,6 +231,18 @@ static void AppendMetadataToString(FString& Str, TArrayView<const uint8>& Metada
 		switch (Context.MajorType())
 		{
 		case ECborCode::Int:
+		{
+			int64 Value = Context.AsInt();
+			if (Value > 999'999'999LL)
+			{
+				Str += FString::Printf(TEXT("0x%llX"), Value);
+			}
+			else
+			{
+				Str += FString::Printf(TEXT("%lld"), Value);
+			}
+			continue;
+		}
 		case ECborCode::Uint:
 		{
 			uint64 Value = Context.AsUInt();
@@ -362,11 +388,15 @@ static void AppendMetadataToString(FString& Str, const TraceServices::FMetadataS
 		switch (Context.MajorType())
 		{
 		case ECborCode::Int:
+			{
+				int64 Value = Context.AsInt();
+				AddValueToName(Data, MaxLength, GetNextFormatSection(), Value);
+				continue;
+			}
 		case ECborCode::Uint:
 			{
 				uint64 Value = Context.AsUInt();
 				AddValueToName(Data, MaxLength, GetNextFormatSection(), Value);
-
 				continue;
 			}
 
diff --git a/Engine/Source/Developer/TraceInsights/Private/Insights/TimingProfiler/ViewModels/TimerFilters.cpp b/Engine/Source/Developer/TraceInsights/Private/Insights/TimingProfiler/ViewModels/TimerFilters.cpp
--- a/Engine/Source/Developer/TraceInsights/Private/Insights/TimingProfiler/ViewModels/TimerFilters.cpp
+++ b/Engine/Source/Developer/TraceInsights/Private/Insights/TimingProfiler/ViewModels/TimerFilters.cpp
@@ -374,7 +374,7 @@ bool FMetadataFilterState::ApplyFilterToMetadata(TArrayView<const uint8>& Metada
 			}
 			ensure(ConvertedData.IsType<int64>());
 
-			int64 MetadataValue = static_cast<int64>(Context.AsUInt());
+			int64 MetadataValue = Context.AsInt();
 			int64 InputValue = ConvertedData.Get<int64>();
 
 			using FFilterOperatorInt64 = FFilterOperator<int64>;
