[Insights] Timing Graph: Added support for Cpu Stack Sampling Instance Graphs

#rnx

diff --git a/Engine/Source/Developer/TraceInsights/Private/Insights/ViewModels/TimingGraphTrack.cpp b/Engine/Source/Developer/TraceInsights/Private/Insights/ViewModels/TimingGraphTrack.cpp
--- a/Engine/Source/Developer/TraceInsights/Private/Insights/ViewModels/TimingGraphTrack.cpp
+++ b/Engine/Source/Developer/TraceInsights/Private/Insights/ViewModels/TimingGraphTrack.cpp
@@ -5,8 +5,10 @@
 #include "Algo/BinarySearch.h"
 
 // TraceServices
+#include "Common/ProviderLock.h"
 #include "TraceServices/Model/Counters.h"
 #include "TraceServices/Model/Frames.h"
+#include "TraceServices/Model/StackSamples.h"
 #include "TraceServices/Model/TimingProfiler.h"
 
 // TraceInsightsCore
@@ -18,7 +20,6 @@
 #include "Insights/TimingProfiler/TimingProfilerManager.h"
 #include "Insights/TimingProfiler/Tracks/ThreadTimingTrack.h"
 #include "Insights/TimingProfiler/ViewModels/FrameStatsHelper.h"
-#include "Insights/TimingProfiler/ViewModels/ThreadTimingSharedState.h"
 #include "Insights/TimingProfiler/Widgets/STimingProfilerWindow.h"
 #include "Insights/ViewModels/AxisViewportDouble.h"
 #include "Insights/ViewModels/GraphTrackBuilder.h"
@@ -459,6 +460,9 @@ void FTimingGraphTrack::UpdateTimerSeries(FTimingGraphSeries& Series, const FTim
 
 		TSet<uint32> Timelines;
 		GetVisibleTimelineIndexes(Timelines);
+		TSet<uint32> VisibleStackThreads;
+		GetVisibleCpuSamplingThreads(VisibleStackThreads);
+
 		const double SessionDuration = Session->GetDurationSeconds();
 		if (Series.CachedSessionDuration != SessionDuration || Series.CachedTimelinesNum != Timelines.Num())
 		{
@@ -475,52 +479,64 @@ void FTimingGraphTrack::UpdateTimerSeries(FTimingGraphSeries& Series, const FTim
 			Series.CachedTimelinesNum = Timelines.Num();
 			Series.CachedEvents.Empty();
 
-			for (uint32 TimelineIndex : Timelines)
+			auto TimelineCallback = [SessionDuration, &Series, TimerReader, &Viewport, &NumTimelinesContainingEvent](const TraceServices::ITimingProfilerProvider::Timeline& Timeline)
 			{
-				TimingProfilerProvider.ReadTimeline(TimelineIndex,
-					[SessionDuration, &Series, TimerReader, &Viewport, &NumTimelinesContainingEvent](const TraceServices::ITimingProfilerProvider::Timeline& Timeline)
+				TArray<TArray<FTimingGraphSeries::FSimpleTimingEvent>> Events;
+				TraceServices::ITimingProfilerTimeline::EnumerateAsyncParams Params;
+				Params.IntervalStart = 0;
+				Params.IntervalEnd = SessionDuration;
+				Params.Resolution = 0.0;
+				Params.SetupCallback = [&Events](uint32 NumTasks)
 					{
-						TArray<TArray<FTimingGraphSeries::FSimpleTimingEvent>> Events;
-						TraceServices::ITimeline<TraceServices::FTimingProfilerEvent>::EnumerateAsyncParams Params;
-						Params.IntervalStart = 0;
-						Params.IntervalEnd = SessionDuration;
-						Params.Resolution = 0.0;
-						Params.SetupCallback = [&Events](uint32 NumTasks)
-						{
-							Events.AddDefaulted(NumTasks);
-						};
-						Params.EventRangeCallback = [&Events, TimerReader, &Series](double StartTime, double EndTime, uint32 Depth, const TraceServices::FTimingProfilerEvent& Event, uint32 TaskIndex)
+						Events.AddDefaulted(NumTasks);
+					};
+				Params.EventRangeCallback = [&Events, TimerReader, &Series](double StartTime, double EndTime, uint32 Depth, const TraceServices::FTimingProfilerEvent& Event, uint32 TaskIndex)
+					{
+						const TraceServices::FTimingProfilerTimer* Timer = TimerReader->GetTimer(Event.TimerIndex);
+						if (ensure(Timer != nullptr))
 						{
-							const TraceServices::FTimingProfilerTimer* Timer = TimerReader->GetTimer(Event.TimerIndex);
-							if (ensure(Timer != nullptr))
+							if (Timer->Id == Series.TimerId)
 							{
-								if (Timer->Id == Series.TimerId)
-								{
-									const double Duration = EndTime - StartTime;
-									Events[TaskIndex].Add({ StartTime, Duration });
-								}
+								const double Duration = EndTime - StartTime;
+								Events[TaskIndex].Add({ StartTime, Duration });
 							}
-							return TraceServices::EEventEnumerate::Continue;
-						};
+						}
+						return TraceServices::EEventEnumerate::Continue;
+					};
 
-						Timeline.EnumerateEventsDownSampledAsync(Params);
+				Timeline.EnumerateEventsDownSampledAsync(Params);
 
-						int32 NumOfCachedEvents = Series.CachedEvents.Num();
-						for (auto& Array : Events)
-						{
-							for (auto& Event : Array)
-							{
-								Series.CachedEvents.Add(Event);
-							}
+				int32 NumOfCachedEvents = Series.CachedEvents.Num();
+				for (auto& Array : Events)
+				{
+					for (auto& Event : Array)
+					{
+						Series.CachedEvents.Add(Event);
+					}
 
-							Array.Empty();
-						}
+					Array.Empty();
+				}
 
-						if (NumOfCachedEvents != Series.CachedEvents.Num())
-						{
-							++NumTimelinesContainingEvent;
-						}
-					});
+				if (NumOfCachedEvents != Series.CachedEvents.Num())
+				{
+					++NumTimelinesContainingEvent;
+				}
+			};
+
+			for (uint32 TimelineIndex : Timelines)
+			{
+				TimingProfilerProvider.ReadTimeline(TimelineIndex, TimelineCallback);
+			}
+
+			const TraceServices::IStackSamplesProvider* StackSamplesProvider = TraceServices::ReadStackSamplesProvider(*Session.Get());
+			if (VisibleStackThreads.Num() > 0 && StackSamplesProvider)
+			{
+				for (uint32 ThreadId : VisibleStackThreads)
+				{
+					TraceServices::FProviderReadScopeLock _(*StackSamplesProvider);
+					const TraceServices::ITimingProfilerProvider::Timeline& Timeline = *StackSamplesProvider->GetTimeline(ThreadId);
+					TimelineCallback(Timeline);
+				}
 			}
 
 			//If events come from multiple timelines, we have to sort the whole thing.
@@ -812,16 +828,37 @@ void FTimingGraphTrack::UpdateStatsCounterSeries(FTimingGraphSeries& Series, con
 
 ////////////////////////////////////////////////////////////////////////////////////////////////////
 
-void FTimingGraphTrack::GetVisibleTimelineIndexes(TSet<uint32>& TimelineIndexes)
+TSharedPtr<UE::Insights::TimingProfiler::FThreadTimingSharedState> FTimingGraphTrack::GetThreadTimingSharedState() const
 {
 	TSharedPtr<UE::Insights::TimingProfiler::STimingView> TimingViewPtr = TimingView.Pin();
 	if (!TimingViewPtr.IsValid())
 	{
-		return;
+		return nullptr;
+	}
+
+	return TimingViewPtr->GetThreadTimingSharedState();
+}
+
+////////////////////////////////////////////////////////////////////////////////////////////////////
+
+void FTimingGraphTrack::GetVisibleTimelineIndexes(TSet<uint32>& TimelineIndexes)
+{
+	TSharedPtr<UE::Insights::TimingProfiler::FThreadTimingSharedState> ThreadSharedState = GetThreadTimingSharedState();
+	if (ThreadSharedState)
+	{
+		ThreadSharedState->GetVisibleTimelineIndexes(TimelineIndexes);
 	}
+}
 
-	TSharedPtr<UE::Insights::TimingProfiler::FThreadTimingSharedState> ThreadSharedState = TimingViewPtr->GetThreadTimingSharedState();
-	ThreadSharedState->GetVisibleTimelineIndexes(TimelineIndexes);
+////////////////////////////////////////////////////////////////////////////////////////////////////
+
+void FTimingGraphTrack::GetVisibleCpuSamplingThreads(TSet<uint32>& Threads)
+{
+	TSharedPtr<UE::Insights::TimingProfiler::FThreadTimingSharedState> ThreadSharedState = GetThreadTimingSharedState();
+	if (ThreadSharedState)
+	{
+		ThreadSharedState->GetVisibleCpuSamplingThreads(Threads);
+	}
 }
 
 ////////////////////////////////////////////////////////////////////////////////////////////////////
diff --git a/Engine/Source/Developer/TraceInsights/Private/Insights/ViewModels/TimingGraphTrack.h b/Engine/Source/Developer/TraceInsights/Private/Insights/ViewModels/TimingGraphTrack.h
--- a/Engine/Source/Developer/TraceInsights/Private/Insights/ViewModels/TimingGraphTrack.h
+++ b/Engine/Source/Developer/TraceInsights/Private/Insights/ViewModels/TimingGraphTrack.h
@@ -5,6 +5,7 @@
 #include "CoreTypes.h"
 
 // TraceInsights
+#include "Insights/TimingProfiler/ViewModels/ThreadTimingSharedState.h"
 #include "Insights/ViewModels/GraphSeries.h"
 #include "Insights/ViewModels/GraphTrack.h"
 
@@ -109,7 +110,9 @@ protected:
 	void UpdateFrameStatsTimerSeries(FTimingGraphSeries& Series, const FTimingTrackViewport& Viewport);
 	void UpdateStatsCounterSeries(FTimingGraphSeries& Series, const FTimingTrackViewport& Viewport);
 
+	TSharedPtr<UE::Insights::TimingProfiler::FThreadTimingSharedState> GetThreadTimingSharedState() const;
 	void GetVisibleTimelineIndexes(TSet<uint32>& TimelineIndexes);
+	void GetVisibleCpuSamplingThreads(TSet<uint32>& Threads);
 
 	virtual void DrawVerticalAxisGrid(const ITimingTrackDrawContext& Context) const override;
 
