Fix Timing Graph empty container

#rnx

diff --git a/Engine/Source/Developer/TraceInsights/Private/Insights/ViewModels/TimingGraphTrack.cpp b/Engine/Source/Developer/TraceInsights/Private/Insights/ViewModels/TimingGraphTrack.cpp
--- a/Engine/Source/Developer/TraceInsights/Private/Insights/ViewModels/TimingGraphTrack.cpp
+++ b/Engine/Source/Developer/TraceInsights/Private/Insights/ViewModels/TimingGraphTrack.cpp
@@ -507,14 +507,18 @@ void FTimingGraphTrack::UpdateTimerSeries(FTimingGraphSeries& Series, const FTim
 				Timeline.EnumerateEventsDownSampledAsync(Params);
 
 				int32 NumOfCachedEvents = Series.CachedEvents.Num();
-				for (auto& Array : Events)
+
+				if (Events.Num())
 				{
-					for (auto& Event : Array)
+					for (auto& Array : Events)
 					{
-						Series.CachedEvents.Add(Event);
-					}
+						for (auto& Event : Array)
+						{
+							Series.CachedEvents.Add(Event);
+						}
 
-					Array.Empty();
+						Array.Empty();
+					}
 				}
 
 				if (NumOfCachedEvents != Series.CachedEvents.Num())
