Fix Timing Graph empty container 2

diff --git a/Engine/Source/Developer/TraceInsights/Private/Insights/ViewModels/TimingGraphTrack.cpp b/Engine/Source/Developer/TraceInsights/Private/Insights/ViewModels/TimingGraphTrack.cpp
--- a/Engine/Source/Developer/TraceInsights/Private/Insights/ViewModels/TimingGraphTrack.cpp
+++ b/Engine/Source/Developer/TraceInsights/Private/Insights/ViewModels/TimingGraphTrack.cpp
@@ -512,12 +512,15 @@ void FTimingGraphTrack::UpdateTimerSeries(FTimingGraphSeries& Series, const FTim
 				{
 					for (auto& Array : Events)
 					{
-						for (auto& Event : Array)
+						if (Array.Num())
 						{
-							Series.CachedEvents.Add(Event);
-						}
+							for (auto& Event : Array)
+							{
+								Series.CachedEvents.Add(Event);
+							}
 
-						Array.Empty();
+							Array.Empty();
+						}
 					}
 				}
 
