Fix Timing Graph empty container 3

#rnx

diff --git a/Engine/Source/Developer/TraceInsights/Private/Insights/ViewModels/TimingGraphTrack.cpp b/Engine/Source/Developer/TraceInsights/Private/Insights/ViewModels/TimingGraphTrack.cpp
--- a/Engine/Source/Developer/TraceInsights/Private/Insights/ViewModels/TimingGraphTrack.cpp
+++ b/Engine/Source/Developer/TraceInsights/Private/Insights/ViewModels/TimingGraphTrack.cpp
@@ -508,11 +508,11 @@ void FTimingGraphTrack::UpdateTimerSeries(FTimingGraphSeries& Series, const FTim
 
 				int32 NumOfCachedEvents = Series.CachedEvents.Num();
 
-				if (Events.Num())
+				if (Events.Num() > 0)
 				{
 					for (auto& Array : Events)
 					{
-						if (Array.Num())
+						if (Array.Num() > 0)
 						{
 							for (auto& Event : Array)
 							{
