Fix Timing Track empty container 4 with suppress

#rnx

diff --git a/Engine/Source/Developer/TraceInsights/Private/Insights/ViewModels/TimingGraphTrack.cpp b/Engine/Source/Developer/TraceInsights/Private/Insights/ViewModels/TimingGraphTrack.cpp
--- a/Engine/Source/Developer/TraceInsights/Private/Insights/ViewModels/TimingGraphTrack.cpp
+++ b/Engine/Source/Developer/TraceInsights/Private/Insights/ViewModels/TimingGraphTrack.cpp
@@ -507,21 +507,14 @@ void FTimingGraphTrack::UpdateTimerSeries(FTimingGraphSeries& Series, const FTim
 				Timeline.EnumerateEventsDownSampledAsync(Params);
 
 				int32 NumOfCachedEvents = Series.CachedEvents.Num();
-
-				if (Events.Num() > 0)
+				for (auto& Array : Events) //-V1078
 				{
-					for (auto& Array : Events)
+					for (auto& Event : Array) //-V1078
 					{
-						if (Array.Num() > 0)
-						{
-							for (auto& Event : Array)
-							{
-								Series.CachedEvents.Add(Event);
-							}
-
-							Array.Empty();
-						}
+						Series.CachedEvents.Add(Event);
 					}
+
+					Array.Empty();
 				}
 
 				if (NumOfCachedEvents != Series.CachedEvents.Num())
