Insights.Client.TimingInsights.OnTimeConnect.DefaultChannels

#rnx

diff --git a/Engine/Source/Developer/TraceInsights/Private/Insights/Tests/InsightsTestUtils.cpp b/Engine/Source/Developer/TraceInsights/Private/Insights/Tests/InsightsTestUtils.cpp
--- a/Engine/Source/Developer/TraceInsights/Private/Insights/Tests/InsightsTestUtils.cpp
+++ b/Engine/Source/Developer/TraceInsights/Private/Insights/Tests/InsightsTestUtils.cpp
@@ -14,6 +14,7 @@
 #include "TraceServices/AnalysisService.h"
 #include "TraceServices/ITraceServicesModule.h"
 #include "TraceServices/Model/AnalysisSession.h"
+#include "TraceServices/Model/Diagnostics.h"
 #include "TraceServices/ModuleService.h"
 
 // TraceInsightsCore
@@ -40,17 +41,21 @@ FInsightsTestUtils::FInsightsTestUtils(FAutomationTestBase* InTest) :
 
 bool FInsightsTestUtils::AnalyzeTrace(const TCHAR* Path) const
 {
-	IUnrealInsightsModule& TraceInsightsModule = FModuleManager::LoadModuleChecked<IUnrealInsightsModule>("TraceInsights");
-
-	if (!FPaths::FileExists(Path))
+	TSharedPtr<const TraceServices::IAnalysisSession> Session = GetSession(Path);
+	if (!Session)
 	{
-		Test->AddError(FString::Printf(TEXT("File does not exist: %s."), Path));
+		Test->AddError(TEXT("Session analysis failed to start."));
 		return false;
 	}
 
-	TraceInsightsModule.StartAnalysisForTraceFile(Path);
-	auto Session = TraceInsightsModule.GetAnalysisSession();
-	if (Session == nullptr)
+	return AnalyzeTrace(Session);
+}
+
+////////////////////////////////////////////////////////////////////////////////////////////////////
+
+bool FInsightsTestUtils::AnalyzeTrace(TSharedPtr<const TraceServices::IAnalysisSession> Session) const
+{
+	if (!Session)
 	{
 		Test->AddError(TEXT("Session analysis failed to start."));
 		return false;
@@ -85,6 +90,23 @@ bool FInsightsTestUtils::AnalyzeTrace(const TCHAR* Path) const
 
 ////////////////////////////////////////////////////////////////////////////////////////////////////
 
+TSharedPtr<const TraceServices::IAnalysisSession> FInsightsTestUtils::GetSession(const TCHAR* TraceFilePath) const
+{
+	IUnrealInsightsModule& TraceInsightsModule = FModuleManager::LoadModuleChecked<IUnrealInsightsModule>("TraceInsights");
+
+	if (!FPaths::FileExists(TraceFilePath))
+	{
+		Test->AddError(FString::Printf(TEXT("File does not exist: %s."), TraceFilePath));
+		return nullptr;
+	}
+
+	TraceInsightsModule.StartAnalysisForTraceFile(TraceFilePath);
+
+	return TraceInsightsModule.GetAnalysisSession();
+}
+
+////////////////////////////////////////////////////////////////////////////////////////////////////
+
 bool FInsightsTestUtils::FileContainsString(const FString& PathToFile, const FString& ExpectedString, double Timeout) const
 {
 	double StartTime = FPlatformTime::Seconds();
@@ -131,7 +153,7 @@ bool FInsightsTestUtils::IsUnrealTraceServerReady(const TCHAR* Host, int32 Port)
 		return false;
 	}
 	const UE::Trace::FStoreClient::FVersion* Version = StoreClient->GetVersion();
-	if (Version == nullptr)
+	if (!Version)
 	{
 		Test->AddError(TEXT("Cannot get version of UTS"));
 		return false;
@@ -301,6 +323,7 @@ void FInsightsTestUtils::ResetSession() const
 }
 
 ////////////////////////////////////////////////////////////////////////////////////////////////////
+
 TSharedPtr<UE::Trace::FStoreClient> FInsightsTestUtils::CreateStoreClient(const TCHAR* Host, int32 Port) const
 {
 	TSharedPtr<UE::Trace::FStoreClient> StoreClient(UE::Trace::FStoreClient::Connect(Host, Port));
@@ -311,6 +334,8 @@ TSharedPtr<UE::Trace::FStoreClient> FInsightsTestUtils::CreateStoreClient(const
 	return StoreClient;
 }
 
+////////////////////////////////////////////////////////////////////////////////////////////////////
+
 uint32 FInsightsTestUtils::GetValidSessionCount(TSharedPtr<UE::Trace::FStoreClient> StoreClient) const
 {
 	uint32 SessionCount = StoreClient->GetSessionCount();
@@ -318,6 +343,8 @@ uint32 FInsightsTestUtils::GetValidSessionCount(TSharedPtr<UE::Trace::FStoreClie
 	return SessionCount;
 }
 
+////////////////////////////////////////////////////////////////////////////////////////////////////
+
 const UE::Trace::FStoreClient::FTraceInfo* FInsightsTestUtils::FindTraceInfoByName(TSharedPtr<UE::Trace::FStoreClient> StoreClient, const FString& TraceName) const
 {
 	uint32 SessionCount = GetValidSessionCount(StoreClient);
@@ -348,6 +375,44 @@ const UE::Trace::FStoreClient::FTraceInfo* FInsightsTestUtils::FindTraceInfoByNa
 	return nullptr;
 }
 
+////////////////////////////////////////////////////////////////////////////////////////////////////
+
+bool FInsightsTestUtils::GetMetadata(const FString& TraceFilePath, TMap<FString, FString>& OutMetadata) const
+{
+	TSharedPtr<const TraceServices::IAnalysisSession> Session = GetSession(*TraceFilePath);
+	if (!Session)
+	{
+		Test->AddError(TEXT("Unable to retrieve metadata from an invalid session."));
+		return false;
+	}
+
+	AnalyzeTrace(Session);
+
+	TraceServices::FAnalysisSessionReadScope SessionReadScope(*Session.Get());
+	const TraceServices::IDiagnosticsProvider* DiagnosticsProvider = TraceServices::ReadDiagnosticsProvider(*Session.Get());
+	if (DiagnosticsProvider && DiagnosticsProvider->IsSessionInfoAvailable())
+	{
+		TraceServices::FSessionInfo SessionInfo = DiagnosticsProvider->GetSessionInfo();
+
+		OutMetadata.Add("Platform", SessionInfo.Platform);
+		OutMetadata.Add("AppName", SessionInfo.AppName);
+		OutMetadata.Add("ProjectName", SessionInfo.ProjectName);
+		OutMetadata.Add("Branch", SessionInfo.Branch);
+		OutMetadata.Add("BuildVersion", SessionInfo.BuildVersion);
+		OutMetadata.Add("Changelist", FString::FromInt(SessionInfo.Changelist));
+		OutMetadata.Add("BuildConfigurationType", LexToString(SessionInfo.ConfigurationType));
+		OutMetadata.Add("BuildTargetType", LexToString(SessionInfo.TargetType));
+		OutMetadata.Add("CommandLine", SessionInfo.CommandLine);
+
+		return true;
+	}
+
+	Test->AddError(TEXT("Metadata map is empty"));
+	return false;
+}
+
+////////////////////////////////////////////////////////////////////////////////////////////////////
+
 bool FInsightsTestUtils::IsTraceHasLiveStatus(const FString& TraceName, const TCHAR* Host, int32 Port) const
 {
 	TSharedPtr<UE::Trace::FStoreClient> StoreClient = CreateStoreClient(Host, Port);
@@ -399,6 +464,8 @@ FString FInsightsTestUtils::GetLiveTrace(const TCHAR* Host, int32 Port) const
 	return TEXT("");
 }
 
+////////////////////////////////////////////////////////////////////////////////////////////////////
+
 FString FInsightsTestUtils::FInsightsTestUtils::GetUTSPath() const
 {
 #if PLATFORM_WINDOWS
diff --git a/Engine/Source/Developer/TraceInsights/Public/Insights/Tests/InsightsTestUtils.h b/Engine/Source/Developer/TraceInsights/Public/Insights/Tests/InsightsTestUtils.h
--- a/Engine/Source/Developer/TraceInsights/Public/Insights/Tests/InsightsTestUtils.h
+++ b/Engine/Source/Developer/TraceInsights/Public/Insights/Tests/InsightsTestUtils.h
@@ -6,6 +6,7 @@
 #include "Trace/StoreClient.h"
 #include "Containers/UnrealString.h"
 #include "ProfilingDebugging/TraceAuxiliary.h" // for FTraceAuxiliary::EConnectionType
+#include "TraceServices/Model/AnalysisSession.h"
 
 #define UE_API TRACEINSIGHTS_API
 
@@ -17,17 +18,20 @@ public:
 	UE_API FInsightsTestUtils(FAutomationTestBase* Test);
 
 	UE_API bool AnalyzeTrace(const TCHAR* Path) const;
+	UE_API bool AnalyzeTrace(TSharedPtr<const TraceServices::IAnalysisSession> Session) const;
 	UE_API bool FileContainsString(const FString& PathToFile, const FString& ExpectedString, double Timeout) const;
 	UE_API bool SetupUTS(double Timeout = 30.0, bool bUseFork = false) const;
 	UE_API bool KillUTS(double Timeout = 30.0) const;
 	UE_API bool StartTracing(FTraceAuxiliary::EConnectionType ConnectionType, double Timeout = 10.0) const;
 	UE_API bool IsUnrealTraceServerReady(const TCHAR* Host = TEXT("localhost"), int32 Port = 0U) const;
 	UE_API bool IsTraceHasLiveStatus(const FString& TraceName, const TCHAR* Host = TEXT("localhost"), int32 Port = 0U) const;
+	UE_API bool GetMetadata(const FString& TraceFilePath, TMap<FString, FString>& OutMetadata) const;
 	UE_API void ResetSession() const;
 	UE_API FString GetLiveTrace(const TCHAR* Host = TEXT("localhost"), int32 Port = 0U) const;
 	UE_API FString GetUTSPath() const;
 
 private:
+	TSharedPtr<const TraceServices::IAnalysisSession> GetSession(const TCHAR* Path) const;
 	TSharedPtr<UE::Trace::FStoreClient> CreateStoreClient(const TCHAR* Host, int32 Port) const;
 	uint32 GetValidSessionCount(TSharedPtr<UE::Trace::FStoreClient> StoreClient) const;
 	const UE::Trace::FStoreClient::FTraceInfo* FindTraceInfoByName(TSharedPtr<UE::Trace::FStoreClient> StoreClient, const FString& TraceName) const;
