[Insights] Moved TimingGraphTrack.cpp/h to TimingProfiler/GraphTracks as a prerequisite for a larger change.

#rnx

diff --git a/Engine/Source/Developer/TraceInsights/Private/Insights/MemoryProfiler/ViewModels/MemorySharedState.cpp b/Engine/Source/Developer/TraceInsights/Private/Insights/MemoryProfiler/ViewModels/MemorySharedState.cpp
--- a/Engine/Source/Developer/TraceInsights/Private/Insights/MemoryProfiler/ViewModels/MemorySharedState.cpp
+++ b/Engine/Source/Developer/TraceInsights/Private/Insights/MemoryProfiler/ViewModels/MemorySharedState.cpp
@@ -27,7 +27,6 @@
 #include "Insights/MemoryProfiler/ViewModels/ReportXmlParser.h"
 #include "Insights/ViewModels/TimingEvent.h"
 #include "Insights/ViewModels/TimingEventSearch.h"
-#include "Insights/ViewModels/TimingGraphTrack.h"
 #include "Insights/ViewModels/TimingTrackViewport.h"
 #include "Insights/ViewModels/TooltipDrawState.h"
 #include "Insights/Widgets/STimingView.h"
diff --git a/Engine/Source/Developer/TraceInsights/Private/Insights/TimingProfiler/GraphTracks/TimingGraphTrack.cpp b/Engine/Source/Developer/TraceInsights/Private/Insights/TimingProfiler/GraphTracks/TimingGraphTrack.cpp
new file mode 100644
--- /dev/null
+++ b/Engine/Source/Developer/TraceInsights/Private/Insights/TimingProfiler/GraphTracks/TimingGraphTrack.cpp
@@ -0,0 +1,1105 @@
+// Copyright Epic Games, Inc. All Rights Reserved.
+
+#include "TimingGraphTrack.h"
+
+#include "Algo/BinarySearch.h"
+
+// TraceServices
+#include "Common/ProviderLock.h"
+#include "TraceServices/Model/Counters.h"
+#include "TraceServices/Model/Frames.h"
+#include "TraceServices/Model/StackSamples.h"
+#include "TraceServices/Model/TimingProfiler.h"
+
+// TraceInsightsCore
+#include "InsightsCore/Common/PaintUtils.h"
+#include "InsightsCore/Common/TimeUtils.h"
+
+// TraceInsights
+#include "Insights/InsightsManager.h"
+#include "Insights/TimingProfiler/TimingProfilerManager.h"
+#include "Insights/TimingProfiler/Tracks/ThreadTimingTrack.h"
+#include "Insights/TimingProfiler/ViewModels/FrameStatsHelper.h"
+#include "Insights/TimingProfiler/Widgets/STimingProfilerWindow.h"
+#include "Insights/ViewModels/AxisViewportDouble.h"
+#include "Insights/ViewModels/GraphTrackBuilder.h"
+#include "Insights/ViewModels/ITimingViewDrawHelper.h"
+#include "Insights/ViewModels/TimingTrackViewport.h"
+#include "Insights/Widgets/STimingView.h"
+
+#include <limits>
+
+#define LOCTEXT_NAMESPACE "GraphTrack"
+
+INSIGHTS_IMPLEMENT_RTTI(FTimingGraphSeries)
+
+////////////////////////////////////////////////////////////////////////////////////////////////////
+// FTimingGraphSeries
+////////////////////////////////////////////////////////////////////////////////////////////////////
+
+FTimingGraphSeries::FTimingGraphSeries(FTimingGraphSeries::ESeriesType InType)
+	: FGraphSeries()
+	, Type(InType)
+	, TimerId(0)
+	, CachedSessionDuration(0.0)
+	, CachedEvents()
+	, bIsTime(InType == FTimingGraphSeries::ESeriesType::Frame ||
+			  InType == FTimingGraphSeries::ESeriesType::Timer ||
+			  InType == FTimingGraphSeries::ESeriesType::FrameStatsTimer)
+	, bIsMemory(false)
+	, bIsFloatingPoint(false)
+{
+}
+
+////////////////////////////////////////////////////////////////////////////////////////////////////
+
+FTimingGraphSeries::~FTimingGraphSeries()
+{
+}
+
+////////////////////////////////////////////////////////////////////////////////////////////////////
+
+FString FTimingGraphSeries::FormatValue(double Value) const
+{
+	using namespace UE::Insights;
+
+	switch (Type)
+	{
+	case FTimingGraphSeries::ESeriesType::Frame:
+		return FString::Printf(TEXT("%s (%g fps)"), *FormatTimeAuto(Value), 1.0 / Value);
+
+	case FTimingGraphSeries::ESeriesType::Timer:
+	case FTimingGraphSeries::ESeriesType::FrameStatsTimer:
+		return FormatTimeAuto(Value);
+
+	case FTimingGraphSeries::ESeriesType::StatsCounter:
+		if (bIsTime)
+		{
+			return FormatTimeAuto(Value);
+		}
+		else if (bIsMemory)
+		{
+			const int64 MemValue = static_cast<int64>(Value);
+			if (MemValue > 0)
+			{
+				if (MemValue < 1024)
+				{
+					return FString::Printf(TEXT("%s bytes"), *FText::AsNumber(MemValue).ToString());
+				}
+				else
+				{
+					FNumberFormattingOptions FormattingOptions;
+					FormattingOptions.MaximumFractionalDigits = 2;
+					return FString::Printf(TEXT("%s (%s bytes)"), *FText::AsMemory(MemValue, &FormattingOptions).ToString(), *FText::AsNumber(MemValue).ToString());
+				}
+			}
+			else if (MemValue == 0)
+			{
+				return TEXT("0");
+			}
+			else
+			{
+				if (-MemValue < 1024)
+				{
+					return FString::Printf(TEXT("-%s bytes"), *FText::AsNumber(-MemValue).ToString());
+				}
+				else
+				{
+					FNumberFormattingOptions FormattingOptions;
+					FormattingOptions.MaximumFractionalDigits = 2;
+					return FString::Printf(TEXT("-%s (-%s bytes)"), *FText::AsMemory(-MemValue, &FormattingOptions).ToString(), *FText::AsNumber(-MemValue).ToString());
+				}
+			}
+		}
+		else if (bIsFloatingPoint)
+		{
+			return FString::Printf(TEXT("%g"), Value);
+		}
+		else
+		{
+			const int64 Int64Value = static_cast<int64>(Value);
+			return FText::AsNumber(Int64Value).ToString();
+		}
+	}
+
+	return FGraphSeries::FormatValue(Value);
+}
+
+////////////////////////////////////////////////////////////////////////////////////////////////////
+
+void FTimingGraphSeries::SetVisibility(bool bOnOff)
+{
+	FGraphSeries::SetVisibility(bOnOff);
+
+	VisibilityChangedDelegate.Broadcast(bOnOff);
+}
+
+////////////////////////////////////////////////////////////////////////////////////////////////////
+// FTimingGraphTrack
+////////////////////////////////////////////////////////////////////////////////////////////////////
+
+INSIGHTS_IMPLEMENT_RTTI(FTimingGraphTrack)
+
+////////////////////////////////////////////////////////////////////////////////////////////////////
+
+FTimingGraphTrack::FTimingGraphTrack(TSharedPtr<UE::Insights::TimingProfiler::STimingView> InTimingView)
+	: FGraphTrack()
+	, TimingView(InTimingView)
+{
+	LoadDefaultSettings();
+
+	// Add non editable options.
+	EnabledOptions |=
+		EGraphOptions::ShowBaseline |
+		EGraphOptions::ShowThresholds |
+		EGraphOptions::ShowVerticalAxisGrid |
+		EGraphOptions::ShowHeader;
+
+	bNotifyTimersOnDestruction = InTimingView->GetName() == FInsightsManagerTabs::TimingProfilerTabId;
+}
+
+////////////////////////////////////////////////////////////////////////////////////////////////////
+
+FTimingGraphTrack::~FTimingGraphTrack()
+{
+	if (OnTrackVisibilityChangedHandle.IsValid())
+	{
+		TSharedPtr<UE::Insights::TimingProfiler::STimingView> TimingViewPtr = TimingView.Pin();
+		if (TimingView.IsValid())
+		{
+			TimingViewPtr->OnTrackVisibilityChanged().Remove(OnTrackVisibilityChangedHandle);
+			TimingViewPtr->OnTrackAdded().Remove(OnTrackAddedHandle);
+			TimingViewPtr->OnTrackRemoved().Remove(OnTrackRemovedHandle);
+		}
+	}
+
+	if (GameFrameSeriesVisibilityHandle.IsValid())
+	{
+		TSharedPtr<FTimingGraphSeries> GameFramesSeries = GetFrameSeries(ETraceFrameType::TraceFrameType_Game);
+		if (GameFramesSeries.IsValid())
+		{
+			GameFramesSeries->VisibilityChangedDelegate.Remove(GameFrameSeriesVisibilityHandle);
+		}
+	}
+
+	if (RenderingFrameSeriesVisibilityHandle.IsValid())
+	{
+		TSharedPtr<FTimingGraphSeries> RenderingFramesSeries = GetFrameSeries(ETraceFrameType::TraceFrameType_Game);
+		if (RenderingFramesSeries.IsValid())
+		{
+			RenderingFramesSeries->VisibilityChangedDelegate.Remove(RenderingFrameSeriesVisibilityHandle);
+		}
+	}
+
+	using namespace UE::Insights::TimingProfiler;
+
+	if (bNotifyTimersOnDestruction)
+	{
+		TSharedPtr<STimingProfilerWindow> ProfilerWindow = FTimingProfilerManager::Get()->GetProfilerWindow();
+		if (ProfilerWindow.IsValid())
+		{
+			TArray<TSharedPtr<FTimingGraphSeries>> TimerSeries;
+			for (TSharedPtr<FGraphSeries>& Series : AllSeries)
+			{
+				TSharedPtr<FTimingGraphSeries> TimingSeries = StaticCastSharedPtr<FTimingGraphSeries>(Series);
+				if (TimingSeries.IsValid() &&
+					(TimingSeries->Type == FTimingGraphSeries::ESeriesType::Timer ||
+					 TimingSeries->Type == FTimingGraphSeries::ESeriesType::FrameStatsTimer))
+				{
+					TimerSeries.Add(TimingSeries);
+				}
+			}
+			AllSeries.Reset();
+			for (TSharedPtr<FTimingGraphSeries>& TimingSeries : TimerSeries)
+			{
+				ProfilerWindow->OnTimerAddedToGraphsChanged(TimingSeries->TimerId);
+			}
+		}
+	}
+}
+
+////////////////////////////////////////////////////////////////////////////////////////////////////
+
+void FTimingGraphTrack::Update(const ITimingTrackUpdateContext& Context)
+{
+	FGraphTrack::Update(Context);
+
+	if (!OnTrackVisibilityChangedHandle.IsValid())
+	{
+		TSharedPtr<UE::Insights::TimingProfiler::STimingView> TimingViewPtr = TimingView.Pin();
+		if (TimingViewPtr.IsValid())
+		{
+			auto OnTrackAddedRemovedLamda = [this](const TSharedPtr<const FBaseTimingTrack> Track)
+			{
+				if (Track->Is<UE::Insights::TimingProfiler::FThreadTimingTrack>())
+				{
+					// If there are more series than the default frame series.
+					if (this->AllSeries.Num() > ETraceFrameType::TraceFrameType_Count)
+					{
+						this->SetDirtyFlag();
+					}
+				}
+			};
+
+			OnTrackAddedHandle = TimingViewPtr->OnTrackAdded().AddLambda(OnTrackAddedRemovedLamda);
+			OnTrackRemovedHandle = TimingViewPtr->OnTrackRemoved().AddLambda(OnTrackAddedRemovedLamda);
+
+			OnTrackVisibilityChangedHandle = TimingViewPtr->OnTrackVisibilityChanged().AddLambda([this]()
+				{
+					this->SetDirtyFlag();
+				});
+		}
+	}
+
+	const bool bIsEntireGraphTrackDirty = IsDirty() || Context.GetViewport().IsHorizontalViewportDirty();
+	bool bNeedsUpdate = bIsEntireGraphTrackDirty;
+
+	if (!bNeedsUpdate)
+	{
+		for (TSharedPtr<FGraphSeries>& Series : AllSeries)
+		{
+			if (Series->IsVisible() && Series->IsDirty())
+			{
+				// At least one series is dirty.
+				bNeedsUpdate = true;
+				break;
+			}
+		}
+	}
+
+	if (bNeedsUpdate)
+	{
+		ClearDirtyFlag();
+
+		NumAddedEvents = 0;
+
+		const FTimingTrackViewport& Viewport = Context.GetViewport();
+
+		for (TSharedPtr<FGraphSeries>& Series : AllSeries)
+		{
+			if (Series->IsVisible() && (bIsEntireGraphTrackDirty || Series->IsDirty()))
+			{
+				// Clear the flag before updating, because the update itself may further need to set the series as dirty.
+				Series->ClearDirtyFlag();
+
+				TSharedPtr<FTimingGraphSeries> TimingSeries = StaticCastSharedPtr<FTimingGraphSeries>(Series);
+				switch (TimingSeries->Type)
+				{
+				case FTimingGraphSeries::ESeriesType::Frame:
+					UpdateFrameSeries(*TimingSeries, Viewport);
+					break;
+
+				case FTimingGraphSeries::ESeriesType::Timer:
+					UpdateTimerSeries(*TimingSeries, Viewport);
+					break;
+
+				case FTimingGraphSeries::ESeriesType::FrameStatsTimer:
+					UpdateFrameStatsTimerSeries(*TimingSeries, Viewport);
+					break;
+
+				case FTimingGraphSeries::ESeriesType::StatsCounter:
+					UpdateStatsCounterSeries(*TimingSeries, Viewport);
+					break;
+				}
+			}
+		}
+
+		UpdateStats();
+	}
+}
+
+////////////////////////////////////////////////////////////////////////////////////////////////////
+// Frame Series
+////////////////////////////////////////////////////////////////////////////////////////////////////
+
+void FTimingGraphTrack::AddDefaultFrameSeries()
+{
+	const FInsightsSettings& Settings = UE::Insights::FInsightsManager::Get()->GetSettings();
+	TSharedPtr<UE::Insights::TimingProfiler::STimingView> TimingViewPtr = TimingView.Pin();
+
+	TSharedRef<FTimingGraphSeries> GameFramesSeries = MakeShared<FTimingGraphSeries>(FTimingGraphSeries::ESeriesType::Frame);
+	GameFramesSeries->SetName(TEXT("Game Frames"));
+	GameFramesSeries->SetDescription(TEXT("Duration of Game frames"));
+	GameFramesSeries->SetColor(FLinearColor(0.3f, 0.3f, 1.0f, 1.0f), FLinearColor(1.0f, 1.0f, 1.0f, 1.0f));
+	GameFramesSeries->FrameType = TraceFrameType_Game;
+	GameFramesSeries->SetBaselineY(SharedValueViewport.GetBaselineY());
+	GameFramesSeries->SetScaleY(SharedValueViewport.GetScaleY());
+	GameFramesSeries->EnableSharedViewport();
+	if (TimingViewPtr.IsValid() && TimingViewPtr->GetName() == FInsightsManagerTabs::TimingProfilerTabId)
+	{
+		GameFramesSeries->SetVisibility(Settings.GetTimingViewMainGraphShowGameFrames());
+		GameFrameSeriesVisibilityHandle = GameFramesSeries->VisibilityChangedDelegate.AddLambda([](bool bOnOff)
+			{
+				FInsightsSettings& Settings = UE::Insights::FInsightsManager::Get()->GetSettings();
+				Settings.SetAndSaveTimingViewMainGraphShowGameFrames(bOnOff);
+			});
+	}
+	AllSeries.Add(GameFramesSeries);
+
+	TSharedRef<FTimingGraphSeries> RenderingFramesSeries = MakeShared<FTimingGraphSeries>(FTimingGraphSeries::ESeriesType::Frame);
+	RenderingFramesSeries->SetName(TEXT("Rendering Frames"));
+	RenderingFramesSeries->SetDescription(TEXT("Duration of Rendering frames"));
+	RenderingFramesSeries->SetColor(FLinearColor(1.0f, 0.3f, 0.3f, 1.0f), FLinearColor(1.0f, 1.0f, 1.0f, 1.0f));
+	RenderingFramesSeries->FrameType = TraceFrameType_Rendering;
+	RenderingFramesSeries->SetBaselineY(SharedValueViewport.GetBaselineY());
+	RenderingFramesSeries->SetScaleY(SharedValueViewport.GetScaleY());
+	RenderingFramesSeries->EnableSharedViewport();
+	if (TimingViewPtr.IsValid() && TimingViewPtr->GetName() == FInsightsManagerTabs::TimingProfilerTabId)
+	{
+		RenderingFramesSeries->SetVisibility(Settings.GetTimingViewMainGraphShowRenderingFrames());
+		RenderingFrameSeriesVisibilityHandle = RenderingFramesSeries->VisibilityChangedDelegate.AddLambda([](bool bOnOff)
+			{
+				FInsightsSettings& Settings = UE::Insights::FInsightsManager::Get()->GetSettings();
+				Settings.SetAndSaveTimingViewMainGraphShowRenderingFrames(bOnOff);
+			});
+	}
+	AllSeries.Add(RenderingFramesSeries);
+}
+
+////////////////////////////////////////////////////////////////////////////////////////////////////
+
+TSharedPtr<FTimingGraphSeries> FTimingGraphTrack::GetFrameSeries(ETraceFrameType FrameType)
+{
+	TSharedPtr<FGraphSeries>* Ptr = AllSeries.FindByPredicate([FrameType](const TSharedPtr<FGraphSeries>& Series)
+		{
+			const TSharedPtr<FTimingGraphSeries> TimingSeries = StaticCastSharedPtr<FTimingGraphSeries>(Series);
+			return TimingSeries->Type == FTimingGraphSeries::ESeriesType::Frame && TimingSeries->FrameType == FrameType;
+		});
+	return (Ptr != nullptr) ? StaticCastSharedPtr<FTimingGraphSeries>(*Ptr) : nullptr;
+}
+
+////////////////////////////////////////////////////////////////////////////////////////////////////
+
+void FTimingGraphTrack::UpdateFrameSeries(FTimingGraphSeries& Series, const FTimingTrackViewport& Viewport)
+{
+	FGraphTrackBuilder Builder(*this, Series, Viewport);
+	TSharedPtr<const TraceServices::IAnalysisSession> Session = UE::Insights::FInsightsManager::Get()->GetSession();
+	if (Session.IsValid())
+	{
+		TraceServices::FAnalysisSessionReadScope SessionReadScope(*Session.Get());
+		const TraceServices::IFrameProvider& FramesProvider = ReadFrameProvider(*Session.Get());
+
+		const TArray64<double>& FrameStartTimes = FramesProvider.GetFrameStartTimes(Series.FrameType);
+
+		const int64 StartLowerBound = Algo::LowerBound(FrameStartTimes, Viewport.GetStartTime());
+		const uint64 StartIndex = (StartLowerBound > 1) ? StartLowerBound - 2 : 0;
+
+		const int64 EndLowerBound = Algo::LowerBound(FrameStartTimes, Viewport.GetEndTime());
+		const uint64 EndIndex = EndLowerBound + 1;
+
+		FramesProvider.EnumerateFrames(Series.FrameType, StartIndex, EndIndex, [&Builder](const TraceServices::FFrame& Frame)
+		{
+			//TODO: add a "frame converter" (i.e. to fps, milliseconds or seconds)
+			const double Duration = Frame.EndTime - Frame.StartTime;
+			Builder.AddEvent(Frame.StartTime, Duration, Duration);
+		});
+	}
+}
+
+////////////////////////////////////////////////////////////////////////////////////////////////////
+// Timer Series
+////////////////////////////////////////////////////////////////////////////////////////////////////
+
+TSharedPtr<FTimingGraphSeries> FTimingGraphTrack::GetTimerSeries(uint32 TimerId)
+{
+	TSharedPtr<FGraphSeries>* Ptr = AllSeries.FindByPredicate([TimerId](const TSharedPtr<FGraphSeries>& Series)
+	{
+		const TSharedPtr<FTimingGraphSeries> TimingSeries = StaticCastSharedPtr<FTimingGraphSeries>(Series);
+		return TimingSeries->Type == FTimingGraphSeries::ESeriesType::Timer && TimingSeries->TimerId == TimerId;
+	});
+	return (Ptr != nullptr) ? StaticCastSharedPtr<FTimingGraphSeries>(*Ptr) : nullptr;
+}
+
+////////////////////////////////////////////////////////////////////////////////////////////////////
+
+TSharedPtr<FTimingGraphSeries> FTimingGraphTrack::AddTimerSeries(uint32 TimerId, FLinearColor Color)
+{
+	TSharedRef<FTimingGraphSeries> Series = MakeShared<FTimingGraphSeries>(FTimingGraphSeries::ESeriesType::Timer);
+
+	Series->SetName(TEXT("<Timer>"));
+	Series->SetDescription(TEXT("Timer series"));
+
+	const FLinearColor BorderColor(Color.R + 0.4f, Color.G + 0.4f, Color.B + 0.4f, 1.0f);
+	Series->SetColor(Color, BorderColor);
+
+	Series->TimerId = TimerId;
+	//Series->CpuOrGpu = ;
+	//Series->TimelineIndex = ;
+
+	// Use shared viewport.
+	Series->SetBaselineY(SharedValueViewport.GetBaselineY());
+	Series->SetScaleY(SharedValueViewport.GetScaleY());
+	Series->EnableSharedViewport();
+
+	Series->CachedSessionDuration = 0.0;
+
+	AllSeries.Add(Series);
+	return Series;
+}
+
+////////////////////////////////////////////////////////////////////////////////////////////////////
+
+void FTimingGraphTrack::RemoveTimerSeries(uint32 TimerId)
+{
+	AllSeries.RemoveAll([TimerId](const TSharedPtr<FGraphSeries>& Series)
+	{
+		const TSharedPtr<FTimingGraphSeries> TimingSeries = StaticCastSharedPtr<FTimingGraphSeries>(Series);
+		return TimingSeries->Type == FTimingGraphSeries::ESeriesType::Timer && TimingSeries->TimerId == TimerId;
+	});
+}
+
+////////////////////////////////////////////////////////////////////////////////////////////////////
+
+void FTimingGraphTrack::UpdateTimerSeries(FTimingGraphSeries& Series, const FTimingTrackViewport& Viewport)
+{
+	FGraphTrackBuilder Builder(*this, Series, Viewport);
+	TSharedPtr<const TraceServices::IAnalysisSession> Session = UE::Insights::FInsightsManager::Get()->GetSession();
+	if (Session.IsValid())
+	{
+		TraceServices::FAnalysisSessionReadScope SessionReadScope(*Session.Get());
+
+		TSet<uint32> Timelines;
+		GetVisibleTimelineIndexes(Timelines);
+		TSet<uint32> VisibleStackThreads;
+		GetVisibleCpuSamplingThreads(VisibleStackThreads);
+
+		const double SessionDuration = Session->GetDurationSeconds();
+		if (Series.CachedSessionDuration != SessionDuration || Series.CachedTimelinesNum != Timelines.Num())
+		{
+			Series.CachedSessionDuration = SessionDuration;
+
+			const TraceServices::ITimingProfilerProvider& TimingProfilerProvider = *TraceServices::ReadTimingProfilerProvider(*Session.Get());
+
+			const TraceServices::ITimingProfilerTimerReader* TimerReader;
+			TimingProfilerProvider.ReadTimers([&TimerReader](const TraceServices::ITimingProfilerTimerReader& Out) { TimerReader = &Out; });
+
+			const uint32 TimelineCount = TimingProfilerProvider.GetTimelineCount();
+			uint32 NumTimelinesContainingEvent = 0;
+
+			Series.CachedTimelinesNum = Timelines.Num();
+			Series.CachedEvents.Empty();
+
+			auto TimelineCallback = [SessionDuration, &Series, TimerReader, &Viewport, &NumTimelinesContainingEvent](const TraceServices::ITimingProfilerProvider::Timeline& Timeline)
+			{
+				TArray<TArray<FTimingGraphSeries::FSimpleTimingEvent>> Events;
+				TraceServices::ITimingProfilerTimeline::EnumerateAsyncParams Params;
+				Params.IntervalStart = 0;
+				Params.IntervalEnd = SessionDuration;
+				Params.Resolution = 0.0;
+				Params.SetupCallback = [&Events](uint32 NumTasks)
+					{
+						Events.AddDefaulted(NumTasks);
+					};
+				Params.EventRangeCallback = [&Events, TimerReader, &Series](double StartTime, double EndTime, uint32 Depth, const TraceServices::FTimingProfilerEvent& Event, uint32 TaskIndex)
+					{
+						const TraceServices::FTimingProfilerTimer* Timer = TimerReader->GetTimer(Event.TimerIndex);
+						if (ensure(Timer != nullptr))
+						{
+							if (Timer->Id == Series.TimerId)
+							{
+								const double Duration = EndTime - StartTime;
+								Events[TaskIndex].Add({ StartTime, Duration });
+							}
+						}
+						return TraceServices::EEventEnumerate::Continue;
+					};
+
+				Timeline.EnumerateEventsDownSampledAsync(Params);
+
+				int32 NumOfCachedEvents = Series.CachedEvents.Num();
+				for (auto& Array : Events) //-V1078
+				{
+					for (auto& Event : Array) //-V1078
+					{
+						Series.CachedEvents.Add(Event);
+					}
+
+					Array.Empty();
+				}
+
+				if (NumOfCachedEvents != Series.CachedEvents.Num())
+				{
+					++NumTimelinesContainingEvent;
+				}
+			};
+
+			for (uint32 TimelineIndex : Timelines)
+			{
+				TimingProfilerProvider.ReadTimeline(TimelineIndex, TimelineCallback);
+			}
+
+			const TraceServices::IStackSamplesProvider* StackSamplesProvider = TraceServices::ReadStackSamplesProvider(*Session.Get());
+			if (VisibleStackThreads.Num() > 0 && StackSamplesProvider)
+			{
+				for (uint32 ThreadId : VisibleStackThreads)
+				{
+					TraceServices::FProviderReadScopeLock _(*StackSamplesProvider);
+					const TraceServices::ITimingProfilerProvider::Timeline& Timeline = *StackSamplesProvider->GetTimeline(ThreadId);
+					TimelineCallback(Timeline);
+				}
+			}
+
+			//If events come from multiple timelines, we have to sort the whole thing.
+			//If they come from a single timeline, they are already sorted.
+			if (NumTimelinesContainingEvent > 1)
+			{
+				Series.CachedEvents.Sort(&FTimingGraphSeries::CompareEventsByStartTime);
+			}
+		}
+
+		int32 StartIndex = Algo::UpperBoundBy(Series.CachedEvents, Viewport.GetStartTime(), &FTimingGraphSeries::FSimpleTimingEvent::StartTime);
+		if (StartIndex > 0)
+		{
+			StartIndex--;
+		}
+		int32 EndIndex = Algo::UpperBoundBy(Series.CachedEvents, Viewport.GetEndTime(), &FTimingGraphSeries::FSimpleTimingEvent::StartTime);
+		if (EndIndex < Series.CachedEvents.Num())
+		{
+			EndIndex++;
+		}
+		for (int32 Index = StartIndex; Index < EndIndex; ++Index)
+		{
+			const FTimingGraphSeries::FSimpleTimingEvent& Event = Series.CachedEvents[Index];
+			Builder.AddEvent(Event.StartTime, Event.Duration, Event.Duration);
+		}
+	}
+}
+
+////////////////////////////////////////////////////////////////////////////////////////////////////
+// Frame Stats Timer Series
+////////////////////////////////////////////////////////////////////////////////////////////////////
+
+TSharedPtr<FTimingGraphSeries> FTimingGraphTrack::GetFrameStatsTimerSeries(uint32 TimerId, ETraceFrameType FrameType)
+{
+	TSharedPtr<FGraphSeries>* Ptr = AllSeries.FindByPredicate([TimerId, FrameType](const TSharedPtr<FGraphSeries>& Series)
+		{
+			const TSharedPtr<FTimingGraphSeries> TimingSeries = StaticCastSharedPtr<FTimingGraphSeries>(Series);
+			return TimingSeries->Type == FTimingGraphSeries::ESeriesType::FrameStatsTimer &&
+				   TimingSeries->TimerId == TimerId &&
+				   TimingSeries->FrameType == FrameType;
+		});
+	return (Ptr != nullptr) ? StaticCastSharedPtr<FTimingGraphSeries>(*Ptr) : nullptr;
+}
+
+////////////////////////////////////////////////////////////////////////////////////////////////////
+
+TSharedPtr<FTimingGraphSeries> FTimingGraphTrack::AddFrameStatsTimerSeries(uint32 TimerId, ETraceFrameType FrameType, FLinearColor Color)
+{
+	TSharedRef<FTimingGraphSeries> Series = MakeShared<FTimingGraphSeries>(FTimingGraphSeries::ESeriesType::FrameStatsTimer);
+
+	Series->SetName(TEXT("<Frame Stats Timer>"));
+	Series->SetDescription(TEXT("Frame Stats Timer series"));
+
+	const FLinearColor BorderColor(Color.R + 0.4f, Color.G + 0.4f, Color.B + 0.4f, 1.0f);
+	Series->SetColor(Color, BorderColor);
+
+	Series->TimerId = TimerId;
+	Series->FrameType = FrameType;
+
+	// Use shared viewport.
+	Series->SetBaselineY(SharedValueViewport.GetBaselineY());
+	Series->SetScaleY(SharedValueViewport.GetScaleY());
+	Series->EnableSharedViewport();
+
+	Series->CachedSessionDuration = 0.0;
+
+	AllSeries.Add(Series);
+	return Series;
+}
+
+////////////////////////////////////////////////////////////////////////////////////////////////////
+
+void FTimingGraphTrack::RemoveFrameStatsTimerSeries(uint32 TimerId, ETraceFrameType FrameType)
+{
+	AllSeries.RemoveAll([TimerId, FrameType](const TSharedPtr<FGraphSeries>& Series)
+		{
+			const TSharedPtr<FTimingGraphSeries> TimingSeries = StaticCastSharedPtr<FTimingGraphSeries>(Series);
+			return TimingSeries->Type == FTimingGraphSeries::ESeriesType::FrameStatsTimer &&
+				   TimingSeries->TimerId == TimerId &&
+				   TimingSeries->FrameType == FrameType;
+		});
+}
+
+////////////////////////////////////////////////////////////////////////////////////////////////////
+
+void FTimingGraphTrack::UpdateFrameStatsTimerSeries(FTimingGraphSeries& Series, const FTimingTrackViewport& Viewport)
+{
+	using namespace UE::Insights::TimingProfiler;
+
+	FGraphTrackBuilder Builder(*this, Series, Viewport);
+	TSharedPtr<const TraceServices::IAnalysisSession> Session = UE::Insights::FInsightsManager::Get()->GetSession();
+	if (Session.IsValid())
+	{
+		TraceServices::FAnalysisSessionReadScope SessionReadScope(*Session.Get());
+
+		TSet<uint32> VisibleTimelines;
+		GetVisibleTimelineIndexes(VisibleTimelines);
+		const double SessionDuration = Session->GetDurationSeconds();
+		if (Series.CachedSessionDuration != SessionDuration || Series.CachedTimelinesNum != VisibleTimelines.Num())
+		{
+			Series.CachedSessionDuration = SessionDuration;
+			Series.CachedTimelinesNum = VisibleTimelines.Num();
+
+			const TraceServices::IFrameProvider& FramesProvider = ReadFrameProvider(*Session.Get());
+
+			Series.FrameStatsCachedEvents.Empty();
+			uint64 FrameCount = FramesProvider.GetFrameCount(ETraceFrameType::TraceFrameType_Game);
+			if (FrameCount == 0)
+			{
+				return;
+			}
+
+			FramesProvider.EnumerateFrames(Series.FrameType, 0ull, FrameCount, [&Series](const TraceServices::FFrame& Frame)
+				{
+					FFrameStatsCachedEvent Event;
+					Event.FrameStartTime = Frame.StartTime;
+					Event.FrameEndTime = Frame.EndTime;
+					Event.Duration.store(0.0f);
+					Series.FrameStatsCachedEvents.Add(Event);
+				});
+
+			FFrameStatsHelper::ComputeFrameStatsForTimer(Series.FrameStatsCachedEvents, Series.TimerId, VisibleTimelines);
+		}
+
+		int32 StartIndex = Algo::UpperBoundBy(Series.FrameStatsCachedEvents, Viewport.GetStartTime(), &FFrameStatsCachedEvent::FrameStartTime);
+		if (StartIndex > 0)
+		{
+			StartIndex--;
+		}
+		int32 EndIndex = Algo::UpperBoundBy(Series.FrameStatsCachedEvents, Viewport.GetEndTime(), &FFrameStatsCachedEvent::FrameStartTime);
+		if (EndIndex < Series.FrameStatsCachedEvents.Num())
+		{
+			EndIndex++;
+		}
+		for (int32 Index = StartIndex; Index < EndIndex; ++Index)
+		{
+			const FFrameStatsCachedEvent& Entry = Series.FrameStatsCachedEvents[Index];
+			Builder.AddEvent(Entry.FrameStartTime, Entry.Duration.load(), Entry.Duration.load());
+		}
+	}
+}
+
+////////////////////////////////////////////////////////////////////////////////////////////////////
+// Stats Counter Series
+////////////////////////////////////////////////////////////////////////////////////////////////////
+
+TSharedPtr<FTimingGraphSeries> FTimingGraphTrack::GetStatsCounterSeries(uint32 CounterId)
+{
+	TSharedPtr<FGraphSeries>* Ptr = AllSeries.FindByPredicate([CounterId](const TSharedPtr<FGraphSeries>& Series)
+	{
+		const TSharedPtr<FTimingGraphSeries> TimingSeries = StaticCastSharedPtr<FTimingGraphSeries>(Series);
+		return TimingSeries->Type == FTimingGraphSeries::ESeriesType::StatsCounter && TimingSeries->CounterId == CounterId;
+	});
+	return (Ptr != nullptr) ? StaticCastSharedPtr<FTimingGraphSeries>(*Ptr) : nullptr;
+}
+
+////////////////////////////////////////////////////////////////////////////////////////////////////
+
+TSharedPtr<FTimingGraphSeries> FTimingGraphTrack::AddStatsCounterSeries(uint32 CounterId, FLinearColor Color)
+{
+	TSharedRef<FTimingGraphSeries> Series = MakeShared<FTimingGraphSeries>(FTimingGraphSeries::ESeriesType::StatsCounter);
+
+	const TCHAR* CounterName = nullptr;
+	bool bIsTime = false;
+	bool bIsMemory = false;
+	bool bIsFloatingPoint = false;
+
+	TSharedPtr<const TraceServices::IAnalysisSession> Session = UE::Insights::FInsightsManager::Get()->GetSession();
+	if (Session.IsValid())
+	{
+		TraceServices::FAnalysisSessionReadScope SessionReadScope(*Session.Get());
+		const TraceServices::ICounterProvider& CountersProvider = TraceServices::ReadCounterProvider(*Session.Get());
+		if (CounterId < CountersProvider.GetCounterCount())
+		{
+			CountersProvider.ReadCounter(CounterId, [&](const TraceServices::ICounter& Counter)
+			{
+				CounterName = Counter.GetName();
+				//bIsTime = (Counter.GetDisplayHint() == TraceServices::CounterDisplayHint_Time);
+				bIsMemory = (Counter.GetDisplayHint() == TraceServices::CounterDisplayHint_Memory);
+				bIsFloatingPoint = Counter.IsFloatingPoint();
+			});
+		}
+	}
+
+	Series->SetName(CounterName != nullptr ? CounterName : TEXT("<StatsCounter>"));
+	Series->SetDescription(TEXT("Stats counter series"));
+
+	FLinearColor BorderColor(Color.R + 0.4f, Color.G + 0.4f, Color.B + 0.4f, 1.0f);
+	Series->SetColor(Color, BorderColor);
+
+	Series->CounterId = CounterId;
+
+	Series->bIsTime = bIsTime;
+	Series->bIsMemory = bIsMemory;
+	Series->bIsFloatingPoint = bIsFloatingPoint;
+
+	Series->SetBaselineY(GetHeight() - 1.0f);
+	Series->SetScaleY(1.0);
+	Series->EnableAutoZoom();
+
+	AllSeries.Add(Series);
+	return Series;
+}
+
+////////////////////////////////////////////////////////////////////////////////////////////////////
+
+void FTimingGraphTrack::RemoveStatsCounterSeries(uint32 CounterId)
+{
+	AllSeries.RemoveAll([CounterId](const TSharedPtr<FGraphSeries>& Series)
+	{
+		const TSharedPtr<FTimingGraphSeries> TimingSeries = StaticCastSharedPtr<FTimingGraphSeries>(Series);
+		return TimingSeries->Type == FTimingGraphSeries::ESeriesType::StatsCounter && TimingSeries->CounterId == CounterId;
+	});
+}
+
+////////////////////////////////////////////////////////////////////////////////////////////////////
+
+void FTimingGraphTrack::UpdateStatsCounterSeries(FTimingGraphSeries& Series, const FTimingTrackViewport& Viewport)
+{
+	FGraphTrackBuilder Builder(*this, Series, Viewport);
+
+	TSharedPtr<const TraceServices::IAnalysisSession> Session = UE::Insights::FInsightsManager::Get()->GetSession();
+	if (Session.IsValid())
+	{
+		TraceServices::FAnalysisSessionReadScope SessionReadScope(*Session.Get());
+		const TraceServices::ICounterProvider& CounterProvider = TraceServices::ReadCounterProvider(*Session.Get());
+		CounterProvider.ReadCounter(Series.CounterId, [this, &Viewport, &Builder, &Series](const TraceServices::ICounter& Counter)
+		{
+			const float TopY = 4.0f;
+			const float BottomY = GetHeight() - 4.0f;
+
+			if (Series.IsAutoZoomEnabled() && TopY < BottomY)
+			{
+				double MinValue =  std::numeric_limits<double>::infinity();
+				double MaxValue = -std::numeric_limits<double>::infinity();
+
+				if (Counter.IsFloatingPoint())
+				{
+					Counter.EnumerateFloatValues(Viewport.GetStartTime(), Viewport.GetEndTime(), true, [&Builder, &MinValue, &MaxValue](double Time, double Value)
+					{
+						if (Value < MinValue)
+						{
+							MinValue = Value;
+						}
+						if (Value > MaxValue)
+						{
+							MaxValue = Value;
+						}
+					});
+				}
+				else
+				{
+					Counter.EnumerateValues(Viewport.GetStartTime(), Viewport.GetEndTime(), true, [&Builder, &MinValue, &MaxValue](double Time, int64 IntValue)
+					{
+						const double Value = static_cast<double>(IntValue);
+						if (Value < MinValue)
+						{
+							MinValue = Value;
+						}
+						if (Value > MaxValue)
+						{
+							MaxValue = Value;
+						}
+					});
+				}
+
+				Series.UpdateAutoZoom(TopY, BottomY, MinValue, MaxValue);
+			}
+
+			if (Counter.IsFloatingPoint())
+			{
+				Counter.EnumerateFloatValues(Viewport.GetStartTime(), Viewport.GetEndTime(), true, [&Builder](double Time, double Value)
+				{
+					//TODO: add a "value unit converter"
+					Builder.AddEvent(Time, 0.0, Value);
+				});
+			}
+			else
+			{
+				Counter.EnumerateValues(Viewport.GetStartTime(), Viewport.GetEndTime(), true, [&Builder](double Time, int64 IntValue)
+				{
+					//TODO: add a "value unit converter"
+					Builder.AddEvent(Time, 0.0, static_cast<double>(IntValue));
+				});
+			}
+		});
+	}
+}
+
+////////////////////////////////////////////////////////////////////////////////////////////////////
+
+TSharedPtr<UE::Insights::TimingProfiler::FThreadTimingSharedState> FTimingGraphTrack::GetThreadTimingSharedState() const
+{
+	TSharedPtr<UE::Insights::TimingProfiler::STimingView> TimingViewPtr = TimingView.Pin();
+	if (!TimingViewPtr.IsValid())
+	{
+		return nullptr;
+	}
+
+	return TimingViewPtr->GetThreadTimingSharedState();
+}
+
+////////////////////////////////////////////////////////////////////////////////////////////////////
+
+void FTimingGraphTrack::GetVisibleTimelineIndexes(TSet<uint32>& TimelineIndexes)
+{
+	TSharedPtr<UE::Insights::TimingProfiler::FThreadTimingSharedState> ThreadSharedState = GetThreadTimingSharedState();
+	if (ThreadSharedState)
+	{
+		ThreadSharedState->GetVisibleTimelineIndexes(TimelineIndexes);
+	}
+}
+
+////////////////////////////////////////////////////////////////////////////////////////////////////
+
+void FTimingGraphTrack::GetVisibleCpuSamplingThreads(TSet<uint32>& Threads)
+{
+	TSharedPtr<UE::Insights::TimingProfiler::FThreadTimingSharedState> ThreadSharedState = GetThreadTimingSharedState();
+	if (ThreadSharedState)
+	{
+		ThreadSharedState->GetVisibleCpuSamplingThreads(Threads);
+	}
+}
+
+////////////////////////////////////////////////////////////////////////////////////////////////////
+
+void FTimingGraphTrack::ContextMenu_ToggleOption_Execute(EGraphOptions Option)
+{
+	FGraphTrack::ContextMenu_ToggleOption_Execute(Option);
+
+	TSharedPtr<UE::Insights::TimingProfiler::STimingView> TimingViewPtr = TimingView.Pin();
+	if (!TimingViewPtr.IsValid())
+	{
+		return;
+	}
+	if (TimingViewPtr->GetName() != FInsightsManagerTabs::TimingProfilerTabId)
+	{
+		return;
+	}
+
+	FInsightsSettings& Settings = UE::Insights::FInsightsManager::Get()->GetSettings();
+	if (EnumHasAnyFlags(Option, EGraphOptions::ShowPoints))
+	{
+		Settings.SetAndSaveTimingViewMainGraphShowPoints(EnumHasAnyFlags(EnabledOptions, EGraphOptions::ShowPoints));
+	}
+	if (EnumHasAnyFlags(Option, EGraphOptions::ShowPointsWithBorder))
+	{
+		Settings.SetAndSaveTimingViewMainGraphShowPointsWithBorder(EnumHasAnyFlags(EnabledOptions, EGraphOptions::ShowPointsWithBorder));
+	}
+	if (EnumHasAnyFlags(Option, EGraphOptions::ShowLines))
+	{
+		Settings.SetAndSaveTimingViewMainGraphShowConnectedLines(EnumHasAnyFlags(EnabledOptions, EGraphOptions::ShowLines));
+	}
+	if (EnumHasAnyFlags(Option, EGraphOptions::ShowPolygon))
+	{
+		Settings.SetAndTimingViewMainGraphShowPolygons(EnumHasAnyFlags(EnabledOptions, EGraphOptions::ShowPolygon));
+	}
+	if (EnumHasAnyFlags(Option, EGraphOptions::UseEventDuration))
+	{
+		Settings.SetAndSaveTimingViewMainGraphShowEventDuration(EnumHasAnyFlags(EnabledOptions, EGraphOptions::UseEventDuration));
+	}
+	if (EnumHasAnyFlags(Option, EGraphOptions::ShowBars))
+	{
+		Settings.SetAndSaveTimingViewMainGraphShowBars(EnumHasAnyFlags(EnabledOptions, EGraphOptions::ShowBars));
+	}
+}
+
+////////////////////////////////////////////////////////////////////////////////////////////////////
+
+void FTimingGraphTrack::LoadDefaultSettings()
+{
+	TSharedPtr<UE::Insights::TimingProfiler::STimingView> TimingViewPtr = TimingView.Pin();
+	if (TimingViewPtr.IsValid() && TimingViewPtr->GetName() == FInsightsManagerTabs::TimingProfilerTabId)
+	{
+		const FInsightsSettings& Settings = UE::Insights::FInsightsManager::Get()->GetSettings();
+		if (Settings.GetTimingViewMainGraphShowPoints())
+		{
+			EnabledOptions |= EGraphOptions::ShowPoints;
+		}
+		if (Settings.GetTimingViewMainGraphShowPointsWithBorder())
+		{
+			EnabledOptions |= EGraphOptions::ShowPointsWithBorder;
+		}
+		if (Settings.GetTimingViewMainGraphShowConnectedLines())
+		{
+			EnabledOptions |= EGraphOptions::ShowLines;
+		}
+		if (Settings.GetTimingViewMainGraphShowPolygons())
+		{
+			EnabledOptions |= EGraphOptions::ShowPolygon;
+		}
+		if (Settings.GetTimingViewMainGraphShowEventDuration())
+		{
+			EnabledOptions |= EGraphOptions::UseEventDuration;
+		}
+		if (Settings.GetTimingViewMainGraphShowBars())
+		{
+			EnabledOptions |= EGraphOptions::ShowBars;
+		}
+	}
+}
+
+////////////////////////////////////////////////////////////////////////////////////////////////////
+
+void FTimingGraphTrack::DrawVerticalAxisGrid(const ITimingTrackDrawContext& Context) const
+{
+	TSharedPtr<FTimingGraphSeries> FirstTimeUnitSeries;
+	for (const TSharedPtr<FGraphSeries>& Series : AllSeries)
+	{
+		if (Series->IsVisible())
+		{
+			TSharedPtr<FTimingGraphSeries> TimingSeries = StaticCastSharedPtr<FTimingGraphSeries>(Series);
+			if (TimingSeries->bIsTime)
+			{
+				FirstTimeUnitSeries = TimingSeries;
+				break;
+			}
+		}
+	}
+	if (!FirstTimeUnitSeries)
+	{
+		return;
+	}
+
+	FAxisViewportDouble ViewportY;
+	ViewportY.SetSize(GetHeight());
+	ViewportY.SetScaleLimits(std::numeric_limits<double>::min(), std::numeric_limits<double>::max());
+	ViewportY.SetScale(SharedValueViewport.GetScaleY());
+	ViewportY.ScrollAtPos(static_cast<float>(SharedValueViewport.GetBaselineY()) - GetHeight());
+
+	const float ViewWidth = Context.GetViewport().GetWidth();
+	const float RoundedViewHeight = FMath::RoundToFloat(GetHeight());
+
+	const float X0 = ViewWidth - 12.0f; // let some space for the vertical scrollbar
+	const float Y0 = GetPosY();
+
+	constexpr float MinDY = 32.0f; // min vertical distance between horizontal grid lines
+	constexpr float TextH = 14.0f; // label height
+
+	UE::Insights::FDrawContext& DrawContext = Context.GetDrawContext();
+	const FSlateBrush* Brush = Context.GetHelper().GetWhiteBrush();
+	//const FSlateFontInfo& Font = Context.GetHelper().GetEventFont();
+	const TSharedRef<FSlateFontMeasure> FontMeasureService = FSlateApplication::Get().GetRenderer()->GetFontMeasureService();
+	const float FontScale = DrawContext.Geometry.Scale;
+
+	const double TopValue = ViewportY.GetValueAtOffset(RoundedViewHeight);
+	const double GridValue = ViewportY.GetValueAtOffset(MinDY);
+	const double BottomValue = ViewportY.GetValueAtOffset(0.0f);
+	const double Delta = GridValue - BottomValue;
+
+	if (Delta > 0.0)
+	{
+		const double Thresholds[] =
+		{
+			1.0e-9,	// 1ns
+			1.0e-8,	// 10ns
+			1.0e-7,	// 100ns
+			1.0e-6,	// 1us
+			1.0e-5,	// 10us
+			0.0001,	// 100us
+			0.001,	// 1ms
+			0.01,	// 10ms
+			0.1,	// 100ms
+			1.0,	// 1s
+			10.0,	// 10s
+			60.0,	// 1m
+			600.0,	// 10m
+			3600.0,	// 1h
+			36000.0,// 10h
+			86400.0	// 1d
+		};
+		constexpr int32 NumThresholds = sizeof(Thresholds) / sizeof(double);
+		int32 Index = static_cast<int32>(Algo::LowerBound(Thresholds, Delta));
+		if (Index > 0)
+		{
+			Index--;
+		}
+		double TickUnit = Thresholds[Index];
+		int64 DeltaTicks = static_cast<int64>(FMath::CeilToDouble(Delta / TickUnit));
+		if (Index < NumThresholds - 1)
+		{
+			const double NextTickUnit = Thresholds[Index + 1];
+			if (NextTickUnit <= static_cast<double>(DeltaTicks + 1) * TickUnit)
+			{
+				TickUnit = NextTickUnit;
+				DeltaTicks = 1;
+			}
+			else if (DeltaTicks != 1 && DeltaTicks != 5 && DeltaTicks % 2 == 1) // prefer even grid values
+			{
+				DeltaTicks++;
+			}
+		}
+		const double Grid = static_cast<double>(DeltaTicks) * TickUnit;
+		const double StartValue = FMath::GridSnap(BottomValue, Grid);
+
+		const FLinearColor GridColor(0.0f, 0.0f, 0.0f, 0.1f);
+		const FLinearColor TextBgColor(0.05f, 0.05f, 0.05f, 1.0f);
+		const FLinearColor TextColor = FirstTimeUnitSeries->GetColor().CopyWithNewOpacity(1.0f);
+
+		for (double Value = StartValue; Value < TopValue; Value += Grid)
+		{
+			const float Y = Y0 + RoundedViewHeight - FMath::RoundToFloat(ViewportY.GetOffsetForValue(Value));
+
+			const FString LabelText = UE::Insights::FormatTimeAuto(Value);
+
+			// Draw horizontal grid line.
+			DrawContext.DrawBox(0, Y, ViewWidth, 1, Brush, GridColor);
+
+			const FVector2D LabelTextSize = FontMeasureService->Measure(LabelText, Font, FontScale) / FontScale;
+			const float LabelX = X0 - static_cast<float>(LabelTextSize.X) - 4.0f;
+			const float LabelY = FMath::Min(Y0 + GetHeight() - TextH, FMath::Max(Y0, Y - TextH / 2));
+
+			// Draw background for value text.
+			DrawContext.DrawBox(LabelX, LabelY, static_cast<float>(LabelTextSize.X) + 4.0f, TextH, Brush, TextBgColor);
+
+			// Draw value text.
+			DrawContext.DrawText(LabelX + 2.0f, LabelY + 1.0f, LabelText, Font, TextColor);
+		}
+
+		DrawContext.LayerId++;
+	}
+}
+
+////////////////////////////////////////////////////////////////////////////////////////////////////
+
+bool FTimingGraphTrack::HasAnySeriesForTimer(uint32 TimerId) const
+{
+	for (const TSharedPtr<FGraphSeries>& Series : AllSeries)
+	{
+		const TSharedPtr<FTimingGraphSeries> TimingSeries = StaticCastSharedPtr<FTimingGraphSeries>(Series);
+		if (TimingSeries.IsValid() &&
+			(TimingSeries->Type == FTimingGraphSeries::ESeriesType::Timer ||
+			 TimingSeries->Type == FTimingGraphSeries::ESeriesType::FrameStatsTimer) &&
+			TimingSeries->TimerId == TimerId)
+		{
+			return true;
+		}
+	}
+
+	return false;
+}
+
+////////////////////////////////////////////////////////////////////////////////////////////////////
+
+uint32 FTimingGraphTrack::GetNumSeriesForTimer(uint32 TimerId) const
+{
+	uint32 NumSeries = 0;
+
+	for (const TSharedPtr<FGraphSeries>& Series : AllSeries)
+	{
+		const TSharedPtr<FTimingGraphSeries> TimingSeries = StaticCastSharedPtr<FTimingGraphSeries>(Series);
+		if (TimingSeries.IsValid() &&
+			(TimingSeries->Type == FTimingGraphSeries::ESeriesType::Timer ||
+			 TimingSeries->Type == FTimingGraphSeries::ESeriesType::FrameStatsTimer) &&
+			TimingSeries->TimerId == TimerId)
+		{
+			++NumSeries;
+		}
+	}
+
+	return NumSeries;
+}
+
+////////////////////////////////////////////////////////////////////////////////////////////////////
+
+#undef LOCTEXT_NAMESPACE
diff --git a/Engine/Source/Developer/TraceInsights/Private/Insights/TimingProfiler/GraphTracks/TimingGraphTrack.h b/Engine/Source/Developer/TraceInsights/Private/Insights/TimingProfiler/GraphTracks/TimingGraphTrack.h
new file mode 100644
--- /dev/null
+++ b/Engine/Source/Developer/TraceInsights/Private/Insights/TimingProfiler/GraphTracks/TimingGraphTrack.h
@@ -0,0 +1,136 @@
+// Copyright Epic Games, Inc. All Rights Reserved.
+
+#pragma once
+
+#include "CoreTypes.h"
+
+// TraceInsights
+#include "Insights/TimingProfiler/ViewModels/ThreadTimingSharedState.h"
+#include "Insights/ViewModels/GraphSeries.h"
+#include "Insights/ViewModels/GraphTrack.h"
+
+namespace UE::Insights::TimingProfiler
+{
+	struct FFrameStatsCachedEvent;
+	class STimingView;
+}
+
+////////////////////////////////////////////////////////////////////////////////////////////////////
+
+/** The delegate to be invoked when a series visibility is changed. */
+DECLARE_MULTICAST_DELEGATE_OneParam(FSeriesVisibilityChangedDelegate, bool bOnOff);
+
+class FTimingGraphSeries : public FGraphSeries
+{
+	INSIGHTS_DECLARE_RTTI(FTimingGraphSeries, FGraphSeries)
+
+public:
+	enum class ESeriesType
+	{
+		Frame,
+		Timer,
+		StatsCounter,
+		FrameStatsTimer
+	};
+
+	struct FSimpleTimingEvent
+	{
+		double StartTime;
+		double Duration;
+	};
+
+public:
+	explicit FTimingGraphSeries(FTimingGraphSeries::ESeriesType Type);
+	virtual ~FTimingGraphSeries();
+
+	virtual FString FormatValue(double Value) const override;
+
+	static bool CompareEventsByStartTime(const FSimpleTimingEvent& EventA, const FSimpleTimingEvent& EventB)
+	{
+		return EventA.StartTime < EventB.StartTime;
+	}
+
+	virtual void SetVisibility(bool bOnOff) override;
+
+public:
+	ESeriesType Type;
+	union
+	{
+		uint32 TimerId;
+		uint32 CounterId;
+	};
+
+	ETraceFrameType FrameType;
+	double CachedSessionDuration;
+	TArray<FSimpleTimingEvent> CachedEvents; // used by Timer series
+	TArray<UE::Insights::TimingProfiler::FFrameStatsCachedEvent> FrameStatsCachedEvents; // used by Frame Stats Timer series
+	uint32 CachedTimelinesNum = 0; // the number of timelines used to gather the data
+
+	bool bIsTime; // the unit for values is [second]
+	bool bIsMemory; // the unit for value is [byte]
+	bool bIsFloatingPoint; // for stats counters
+
+	FSeriesVisibilityChangedDelegate VisibilityChangedDelegate;
+};
+
+////////////////////////////////////////////////////////////////////////////////////////////////////
+
+class FTimingGraphTrack : public FGraphTrack
+{
+	INSIGHTS_DECLARE_RTTI(FTimingGraphTrack, FGraphTrack)
+
+public:
+	FTimingGraphTrack(TSharedPtr<UE::Insights::TimingProfiler::STimingView> InTimingView);
+	virtual ~FTimingGraphTrack();
+
+	virtual void Update(const ITimingTrackUpdateContext& Context) override;
+
+	void AddDefaultFrameSeries();
+
+	TSharedPtr<FTimingGraphSeries> GetFrameSeries(ETraceFrameType FrameType);
+
+	TSharedPtr<FTimingGraphSeries> GetTimerSeries(uint32 TimerId);
+	TSharedPtr<FTimingGraphSeries> AddTimerSeries(uint32 TimerId, FLinearColor Color);
+	void RemoveTimerSeries(uint32 TimerId);
+
+	TSharedPtr<FTimingGraphSeries> GetFrameStatsTimerSeries(uint32 TimerId, ETraceFrameType FrameType);
+	TSharedPtr<FTimingGraphSeries> AddFrameStatsTimerSeries(uint32 TimerId, ETraceFrameType FrameType, FLinearColor Color);
+	void RemoveFrameStatsTimerSeries(uint32 TimerId, ETraceFrameType FrameType);
+
+	TSharedPtr<FTimingGraphSeries> GetStatsCounterSeries(uint32 CounterId);
+	TSharedPtr<FTimingGraphSeries> AddStatsCounterSeries(uint32 CounterId, FLinearColor Color);
+	void RemoveStatsCounterSeries(uint32 CounterId);
+
+	bool HasAnySeriesForTimer(uint32 TimerId) const;
+	uint32 GetNumSeriesForTimer(uint32 TimerId) const;
+
+protected:
+	void UpdateFrameSeries(FTimingGraphSeries& Series, const FTimingTrackViewport& Viewport);
+	void UpdateTimerSeries(FTimingGraphSeries& Series, const FTimingTrackViewport& Viewport);
+	void UpdateFrameStatsTimerSeries(FTimingGraphSeries& Series, const FTimingTrackViewport& Viewport);
+	void UpdateStatsCounterSeries(FTimingGraphSeries& Series, const FTimingTrackViewport& Viewport);
+
+	TSharedPtr<UE::Insights::TimingProfiler::FThreadTimingSharedState> GetThreadTimingSharedState() const;
+	void GetVisibleTimelineIndexes(TSet<uint32>& TimelineIndexes);
+	void GetVisibleCpuSamplingThreads(TSet<uint32>& Threads);
+
+	virtual void DrawVerticalAxisGrid(const ITimingTrackDrawContext& Context) const override;
+
+	void LoadDefaultSettings();
+
+private:
+	virtual void ContextMenu_ToggleOption_Execute(EGraphOptions Option);
+
+private:
+	FDelegateHandle OnTrackVisibilityChangedHandle;
+	FDelegateHandle OnTrackAddedHandle;
+	FDelegateHandle OnTrackRemovedHandle;
+
+	FDelegateHandle GameFrameSeriesVisibilityHandle;
+	FDelegateHandle RenderingFrameSeriesVisibilityHandle;
+
+	TWeakPtr<UE::Insights::TimingProfiler::STimingView> TimingView;
+	bool bNotifyTimersOnDestruction;
+};
+
+////////////////////////////////////////////////////////////////////////////////////////////////////
diff --git a/Engine/Source/Developer/TraceInsights/Private/Insights/TimingProfiler/Widgets/SStatsView.cpp b/Engine/Source/Developer/TraceInsights/Private/Insights/TimingProfiler/Widgets/SStatsView.cpp
--- a/Engine/Source/Developer/TraceInsights/Private/Insights/TimingProfiler/Widgets/SStatsView.cpp
+++ b/Engine/Source/Developer/TraceInsights/Private/Insights/TimingProfiler/Widgets/SStatsView.cpp
@@ -37,6 +37,7 @@
 // TraceInsights
 #include "Insights/InsightsStyle.h"
 #include "Insights/Log.h"
+#include "Insights/TimingProfiler/GraphTracks/TimingGraphTrack.h"
 #include "Insights/TimingProfiler/TimingProfilerManager.h"
 #include "Insights/TimingProfiler/ViewModels/CounterAggregator.h"
 #include "Insights/TimingProfiler/ViewModels/StatsNodeHelper.h"
@@ -46,7 +47,6 @@
 #include "Insights/TimingProfiler/Widgets/SStatsViewTooltip.h"
 #include "Insights/TimingProfiler/Widgets/STimingProfilerWindow.h"
 #include "Insights/TimingProfilerCommon.h"
-#include "Insights/ViewModels/TimingGraphTrack.h"
 #include "Insights/Widgets/STimingView.h"
 #include "Insights/Table/ViewModels/TableCommands.h"
 
diff --git a/Engine/Source/Developer/TraceInsights/Private/Insights/TimingProfiler/Widgets/STimerTreeView.cpp b/Engine/Source/Developer/TraceInsights/Private/Insights/TimingProfiler/Widgets/STimerTreeView.cpp
--- a/Engine/Source/Developer/TraceInsights/Private/Insights/TimingProfiler/Widgets/STimerTreeView.cpp
+++ b/Engine/Source/Developer/TraceInsights/Private/Insights/TimingProfiler/Widgets/STimerTreeView.cpp
@@ -24,6 +24,7 @@
 // TraceInsights
 #include "Insights/InsightsStyle.h"
 #include "Insights/Table/ViewModels/TableCommands.h"
+#include "Insights/TimingProfiler/GraphTracks/TimingGraphTrack.h"
 #include "Insights/TimingProfiler/TimingProfilerManager.h"
 #include "Insights/TimingProfiler/ViewModels/TimerButterflyAggregator.h"
 #include "Insights/TimingProfiler/ViewModels/TimersViewColumnFactory.h"
@@ -31,7 +32,6 @@
 #include "Insights/TimingProfiler/Widgets/STimersViewTooltip.h"
 #include "Insights/TimingProfiler/Widgets/STimerTableRow.h"
 #include "Insights/TimingProfiler/Widgets/STimingProfilerWindow.h"
-#include "Insights/ViewModels/TimingGraphTrack.h"
 #include "Insights/Widgets/STimingView.h"
 
 #define LOCTEXT_NAMESPACE "UE::Insights::TimingProfiler::STimerTreeView"
diff --git a/Engine/Source/Developer/TraceInsights/Private/Insights/TimingProfiler/Widgets/STimersView.cpp b/Engine/Source/Developer/TraceInsights/Private/Insights/TimingProfiler/Widgets/STimersView.cpp
--- a/Engine/Source/Developer/TraceInsights/Private/Insights/TimingProfiler/Widgets/STimersView.cpp
+++ b/Engine/Source/Developer/TraceInsights/Private/Insights/TimingProfiler/Widgets/STimersView.cpp
@@ -44,6 +44,7 @@
 #include "Insights/InsightsStyle.h"
 #include "Insights/Log.h"
 #include "Insights/Table/ViewModels/TableCommands.h"
+#include "Insights/TimingProfiler/GraphTracks/TimingGraphTrack.h"
 #include "Insights/TimingProfiler/TimingProfilerManager.h"
 #include "Insights/TimingProfiler/Tracks/GpuTimingTrack.h"
 #include "Insights/TimingProfiler/Tracks/VerseTimingTrack.h"
@@ -56,7 +57,6 @@
 #include "Insights/TimingProfiler/Widgets/STimerTableRow.h"
 #include "Insights/TimingProfiler/Widgets/STimingProfilerWindow.h"
 #include "Insights/TimingProfilerCommon.h"
-#include "Insights/ViewModels/TimingGraphTrack.h"
 #include "Insights/Widgets/STimingView.h"
 
 #include <limits>
diff --git a/Engine/Source/Developer/TraceInsights/Private/Insights/TimingProfiler/Widgets/STimingProfilerWindow.cpp b/Engine/Source/Developer/TraceInsights/Private/Insights/TimingProfiler/Widgets/STimingProfilerWindow.cpp
--- a/Engine/Source/Developer/TraceInsights/Private/Insights/TimingProfiler/Widgets/STimingProfilerWindow.cpp
+++ b/Engine/Source/Developer/TraceInsights/Private/Insights/TimingProfiler/Widgets/STimingProfilerWindow.cpp
@@ -11,13 +11,13 @@
 
 // TraceInsights
 #include "Insights/InsightsStyle.h"
+#include "Insights/TimingProfiler/GraphTracks/TimingGraphTrack.h"
 #include "Insights/TimingProfiler/TimingProfilerManager.h"
 #include "Insights/TimingProfiler/Widgets/SFrameTrack.h"
 #include "Insights/TimingProfiler/Widgets/SStatsView.h"
 #include "Insights/TimingProfiler/Widgets/STimersView.h"
 #include "Insights/TimingProfiler/Widgets/STimerTreeView.h"
 #include "Insights/TimingProfiler/Widgets/STimingProfilerToolbar.h"
-#include "Insights/ViewModels/TimingGraphTrack.h"
 #include "Insights/Widgets/SLogView.h"
 #include "Insights/Widgets/SModulesView.h"
 #include "Insights/Widgets/STimingView.h"
diff --git a/Engine/Source/Developer/TraceInsights/Private/Insights/ViewModels/TimingGraphTrack.cpp b/Engine/Source/Developer/TraceInsights/Private/Insights/ViewModels/TimingGraphTrack.cpp
deleted file mode 100644
diff --git a/Engine/Source/Developer/TraceInsights/Private/Insights/ViewModels/TimingGraphTrack.h b/Engine/Source/Developer/TraceInsights/Private/Insights/ViewModels/TimingGraphTrack.h
deleted file mode 100644
diff --git a/Engine/Source/Developer/TraceInsights/Private/Insights/Widgets/STimingView.cpp b/Engine/Source/Developer/TraceInsights/Private/Insights/Widgets/STimingView.cpp
--- a/Engine/Source/Developer/TraceInsights/Private/Insights/Widgets/STimingView.cpp
+++ b/Engine/Source/Developer/TraceInsights/Private/Insights/Widgets/STimingView.cpp
@@ -60,6 +60,7 @@
 #include "Insights/Log.h"
 #include "Insights/TaskGraphProfiler/TaskGraphProfilerManager.h"
 #include "Insights/Tests/TimingProfilerTests.h"
+#include "Insights/TimingProfiler/GraphTracks/TimingGraphTrack.h"
 #include "Insights/TimingProfiler/TimingProfilerManager.h"
 #include "Insights/TimingProfiler/Tracks/FileActivityTimingTrack.h"
 #include "Insights/TimingProfiler/Tracks/MarkersTimingTrack.h"
@@ -83,7 +84,6 @@
 #include "Insights/ViewModels/QuickFind.h"
 #include "Insights/ViewModels/ThreadTrackEvent.h"
 #include "Insights/ViewModels/TimingEventSearch.h"
-#include "Insights/ViewModels/TimingGraphTrack.h"
 #include "Insights/ViewModels/TimingViewDrawHelper.h"
 #include "Insights/Widgets/SLogView.h"
 #include "Insights/Widgets/SQuickFind.h"
