[Insights] Modules View: Added Symbols Information Tooltip
* Added a tooltip to warn users about including exe,dll,pdb files in the same directory for symbol resolution
* Added a row in the tooltip to display currently declared paths

diff --git a/Engine/Source/Developer/TraceInsights/Private/Insights/Widgets/SModulesView.cpp b/Engine/Source/Developer/TraceInsights/Private/Insights/Widgets/SModulesView.cpp
--- a/Engine/Source/Developer/TraceInsights/Private/Insights/Widgets/SModulesView.cpp
+++ b/Engine/Source/Developer/TraceInsights/Private/Insights/Widgets/SModulesView.cpp
@@ -7,6 +7,8 @@
 #include "SlateOptMacros.h"
 #include "Textures/SlateIcon.h"
 #include "Widgets/Layout/SScrollBox.h"
+#include "Widgets/Layout/SSeparator.h"
+#include "Widgets/SToolTip.h"
 #include "Widgets/Text/STextBlock.h"
 
 // TraceServices
@@ -759,6 +761,7 @@ void SModulesView::Construct(const FArguments& InArgs)
 		[
 			SNew(SHorizontalBox)
 
+			// Module stats
 			+ SHorizontalBox::Slot()
 			.VAlign(VAlign_Center)
 			.Padding(2.0f, 0.0f)
@@ -767,6 +770,115 @@ void SModulesView::Construct(const FArguments& InArgs)
 				SNew(STextBlock)
 				.Text(this, &SModulesView::GetStatsText)
 			]
+
+			// Search Path Information
+			+ SHorizontalBox::Slot()
+			.VAlign(VAlign_Center)
+			.HAlign(HAlign_Right)
+			.Padding(2.0f, 0.0f)
+			.FillWidth(1)
+			[
+				SNew(SHorizontalBox)
+				.ToolTip(
+					SNew(SToolTip)
+					[
+						SNew(SHorizontalBox)
+
+						+ SHorizontalBox::Slot()
+						.AutoWidth()
+						[
+							SNew(SVerticalBox)
+
+							+ SVerticalBox::Slot()
+							.AutoHeight()
+							.Padding(2.0f)
+							[
+								SNew(SSeparator)
+								.Orientation(Orient_Horizontal)
+							]
+
+							+ SVerticalBox::Slot()
+							.AutoHeight()
+							.Padding(2.0f)
+							[
+								SNew(SHorizontalBox)
+
+								// Directory Information
+								+ SHorizontalBox::Slot()
+								.AutoWidth()
+								.Padding(2.0f)
+								[
+									SNew(STextBlock)
+									.Text(LOCTEXT("TT_SymbolDirectoryWarning", ".EXE/.DLL and .PDB files must all be located in the same folder."))
+									.TextStyle(FInsightsStyle::Get(), TEXT("TreeTable.TooltipBold"))
+								]
+							]
+
+							+ SVerticalBox::Slot()
+							.AutoHeight()
+							.Padding(2.0f)
+							[
+								SNew(SSeparator)
+								.Orientation(Orient_Horizontal)
+							]
+
+							+ SVerticalBox::Slot()
+							.AutoHeight()
+							.Padding(2.0f)
+							[
+								SNew(SHorizontalBox)
+
+								// Search Paths
+								+ SHorizontalBox::Slot()
+								.AutoWidth()
+								.Padding(2.0f)
+								[
+									SNew(STextBlock)
+									.Text(LOCTEXT("TT_SearchPaths", "Search Path(s):"))
+									.TextStyle(FInsightsStyle::Get(), TEXT("TreeTable.TooltipBold"))
+								]
+
+								+ SHorizontalBox::Slot()
+								.AutoWidth()
+								.Padding(2.0f)
+								[
+									SNew(STextBlock)
+									.Text(this, &SModulesView::GetSymbolSearchPaths)
+									.TextStyle(FInsightsStyle::Get(), TEXT("TreeTable.Tooltip"))
+								]
+							]
+
+							+ SVerticalBox::Slot()
+							.AutoHeight()
+							.Padding(2.0f)
+							[
+								SNew(SSeparator)
+								.Orientation(Orient_Horizontal)
+							]
+						]
+					]
+				)
+
+				+ SHorizontalBox::Slot()
+				.VAlign(VAlign_Center)
+				.HAlign(HAlign_Center)
+				.Padding(2.0f, 0.0f)
+				.AutoWidth()
+				[
+					SNew(SImage)
+					.Image(FInsightsStyle::GetBrush("Icons.Hint.TreeItem"))
+				]
+				
+				+ SHorizontalBox::Slot()
+				.VAlign(VAlign_Center)
+				.HAlign(HAlign_Center)
+				.Padding(2.0f, 0.0f)
+				.AutoWidth()
+				[
+					SNew(STextBlock)
+					.Text(LOCTEXT("SearchPaths", "Search Paths"))
+				]
+			]
 		]
 	];
 }
@@ -857,6 +969,33 @@ FSlateColor SModulesView::GetStatsTextColor() const
 {
 	return FSlateColor(FLinearColor(1.0f, 1.0f, 1.0f, 1.0f));
 }
+	
+////////////////////////////////////////////////////////////////////////////////////////////////////
+
+FText SModulesView::GetSymbolSearchPaths() const
+{
+	const TSharedPtr<const TraceServices::IAnalysisSession> Session = FInsightsManager::Get()->GetSession();
+	if (const TraceServices::IModuleProvider* ModuleProvider = Session->ReadProvider<TraceServices::IModuleProvider>("ModuleProvider"))
+	{
+		if (ModuleProvider)
+		{
+			TStringBuilder<2048> PathsListStr;
+			ModuleProvider->EnumerateSymbolSearchPaths([&](FStringView Path)
+			{
+				FString CleanedPath = FString(Path);
+				CleanedPath.TrimQuotesInline();
+				PathsListStr << (TEXT("%s"), PathsListStr.Len() ? "\n" : "") << TEXT("\u2022 ") << CleanedPath;
+			});
+			
+			if (PathsListStr.Len())
+			{
+				return FText::FromStringView(PathsListStr);
+			}
+		}
+	}
+	
+	return LOCTEXT("SymbolResolutionNoPaths", " No symbol paths have been setup. \n Use configuration file, UE_INSIGHTS_SYMBOL_PATH environment variable or manually select a file or directory from the Modules view.");
+}
 
 ////////////////////////////////////////////////////////////////////////////////////////////////////
 
diff --git a/Engine/Source/Developer/TraceInsights/Private/Insights/Widgets/SModulesView.h b/Engine/Source/Developer/TraceInsights/Private/Insights/Widgets/SModulesView.h
--- a/Engine/Source/Developer/TraceInsights/Private/Insights/Widgets/SModulesView.h
+++ b/Engine/Source/Developer/TraceInsights/Private/Insights/Widgets/SModulesView.h
@@ -81,6 +81,8 @@ protected:
 	FText GetStatsText() const;
 	FSlateColor GetStatsTextColor() const;
 
+	FText GetSymbolSearchPaths() const;
+
 	TSharedPtr<SWidget> ListView_GetContextMenu();
 	FText LoadSymbolsFile_Label() const;
 	FText LoadSymbolsDirectory_Label() const;
