[Insights] Persistent Settings: Moved 'SymbolSearchPaths' into the [Insights.Symbols] section
* Moved from MemoryProfiler to Symbols
* Added a helper class for all Resolvers to get Search Paths from

#rnx

diff --git a/Engine/Programs/UnrealInsights/Config/DefaultUnrealInsightsSettings.ini b/Engine/Programs/UnrealInsights/Config/DefaultUnrealInsightsSettings.ini
--- a/Engine/Programs/UnrealInsights/Config/DefaultUnrealInsightsSettings.ini
+++ b/Engine/Programs/UnrealInsights/Config/DefaultUnrealInsightsSettings.ini
@@ -6,6 +6,14 @@ bCanSaveAllSections=True
 +Name=LoadingProfiler
 +Name=MemoryProfiler
 
+[Insights.Symbols]
+
+; Add symbol search paths by adding to this array. Order is important.
+;+SymbolSearchPaths=C:\Symbols
+; Symbol server format is also accepted, such as the official Microsoft symbol servers (see https://docs.microsoft.com/en-us/windows-hardware/drivers/debugger/advanced-symsrv-use for details)
+;+SymbolSearchPaths=srv*C:\Symbols;https://msdl.microsoft.com/download/symbols
+
+
 [Insights.TimingProfiler]
 
 ; The default (initial) zoom level of the Timing view.
@@ -179,11 +187,6 @@ ShowNonTargetProcessEvents=True
 
 [Insights.MemoryProfiler]
 
-; Add symbol search paths by adding to this array. Order is important.
-;+SymbolSearchPaths=C:\Symbols
-; Symbol server format is also accepted, such as the official Microsoft symbol servers (see https://docs.microsoft.com/en-us/windows-hardware/drivers/debugger/advanced-symsrv-use for details)
-;+SymbolSearchPaths=srv*C:\Symbols;https://msdl.microsoft.com/download/symbols
-
 ; Memory Budgets (used by the LLM Tags view)
 +BudgetFilePath=(Label="Default",Path="Engine/Programs/UnrealInsights/Config/DefaultMemoryBudgets.xml")
 
diff --git a/Engine/Source/Developer/TraceInsights/Private/Insights/InsightsSettings.cpp b/Engine/Source/Developer/TraceInsights/Private/Insights/InsightsSettings.cpp
--- a/Engine/Source/Developer/TraceInsights/Private/Insights/InsightsSettings.cpp
+++ b/Engine/Source/Developer/TraceInsights/Private/Insights/InsightsSettings.cpp
@@ -114,6 +114,12 @@ void FInsightsSettings::LoadFromConfig()
 			}
 		}
 	}
+
+	//////////////////////////////////////////////////
+	// [Insights.Symbols]
+
+	GConfig->GetArray(TEXT("Insights.Symbols"), TEXT("SymbolSearchPaths"), SymbolSearchPaths, SettingsIni);
+
 	//////////////////////////////////////////////////
 	// [Insights.TimingProfiler]
 
@@ -226,11 +232,6 @@ void FInsightsSettings::LoadFromConfig()
 	GConfig->GetBool(TEXT("Insights.TimingProfiler.TimersView"), TEXT("ShowCpuTimers"), bTimersViewShowCpuTimers, SettingsIni);
 	GConfig->GetBool(TEXT("Insights.TimingProfiler.TimersView"), TEXT("ShowCpuSamplingTimers"), bTimersViewShowCpuSamplingTimers, SettingsIni);
 	GConfig->GetBool(TEXT("Insights.TimingProfiler.TimersView"), TEXT("ShowZeroCountTimers"), bTimersViewShowZeroCountTimers, SettingsIni);
-
-	//////////////////////////////////////////////////
-	// [Insights.MemoryProfiler]
-
-	GConfig->GetArray(TEXT("Insights.MemoryProfiler"), TEXT("SymbolSearchPaths"), SymbolSearchPaths, SettingsIni);
 }
 
 ////////////////////////////////////////////////////////////////////////////////////////////////////
@@ -289,6 +290,12 @@ void FInsightsSettings::SaveToConfig()
 			}
 		}
 	}
+
+	//////////////////////////////////////////////////
+	// [Insights.Symbols]
+
+	GConfig->SetArray(TEXT("Insights.Symbols"), TEXT("SymbolSearchPaths"), SymbolSearchPaths, SettingsIni);
+
 	//////////////////////////////////////////////////
 	// [Insights.TimingProfiler]
 
@@ -362,11 +369,6 @@ void FInsightsSettings::SaveToConfig()
 	GConfig->SetBool(TEXT("Insights.TimingProfiler.TimersView"), TEXT("ShowCpuSamplingTimers"), bTimersViewShowCpuSamplingTimers, SettingsIni);
 	GConfig->SetBool(TEXT("Insights.TimingProfiler.TimersView"), TEXT("ShowZeroCountTimers"), bTimersViewShowZeroCountTimers, SettingsIni);
 
-	//////////////////////////////////////////////////
-	// [Insights.MemoryProfiler]
-
-	GConfig->SetArray(TEXT("Insights.MemoryProfiler"), TEXT("SymbolSearchPaths"), SymbolSearchPaths, SettingsIni);
-
 	//////////////////////////////////////////////////
 
 	GConfig->Flush(false, SettingsIni);
diff --git a/Engine/Source/Developer/TraceInsights/Private/Insights/InsightsSettings.h b/Engine/Source/Developer/TraceInsights/Private/Insights/InsightsSettings.h
--- a/Engine/Source/Developer/TraceInsights/Private/Insights/InsightsSettings.h
+++ b/Engine/Source/Developer/TraceInsights/Private/Insights/InsightsSettings.h
@@ -329,6 +329,13 @@ public:
 		SaveToConfig();
 	}
 
+	//////////////////////////////////////////////////
+	// [Insights.Symbols]
+
+	const TArray<FString>& GetSymbolSearchPaths() const { return SymbolSearchPaths; }
+	void SetSymbolSearchPaths(const TArray<FString>& SearchPaths) { SymbolSearchPaths = SearchPaths; }
+	void SetAndSaveSymbolSearchPaths(const TArray<FString>& SearchPaths) { SET_AND_SAVE(SymbolSearchPaths, SearchPaths); }
+
 	//////////////////////////////////////////////////
 	// [Insights.TimingProfiler]
 
@@ -502,13 +509,6 @@ public:
 	void SetTimersViewShowZeroCountTimers(bool InValue) { bTimersViewShowZeroCountTimers = InValue; }
 	void SetAndSaveTimersViewShowZeroCountTimers(bool InValue) { SET_AND_SAVE(bTimersViewShowZeroCountTimers, InValue); }
 
-	//////////////////////////////////////////////////
-	// [Insights.MemoryProfiler]
-
-	const TArray<FString>& GetSymbolSearchPaths() const { return SymbolSearchPaths; }
-	void SetSymbolSearchPaths(const TArray<FString>& SearchPaths) { SymbolSearchPaths = SearchPaths; }
-	void SetAndSaveSymbolSearchPaths(const TArray<FString>& SearchPaths) { SET_AND_SAVE(SymbolSearchPaths, SearchPaths); }
-
 	//////////////////////////////////////////////////
 
 	#undef SET_AND_SAVE
@@ -535,6 +535,12 @@ private:
 	/** The default settings for a Timing View. */
 	inline static UE::Insights::FTimingViewSettings DefaultTimingViewSettings = { };
 
+	//////////////////////////////////////////////////
+	// [Insights.Symbols]
+
+	/** List of search paths to look for symbol files */
+	TArray<FString> SymbolSearchPaths;
+
 	//////////////////////////////////////////////////
 	// [Insights.TimingProfiler]
 
@@ -655,11 +661,5 @@ private:
 	/** If enabled, timers with no instances in the selected interval will still be displayed in the Timers View. */
 	bool bTimersViewShowZeroCountTimers = true;
 
-	//////////////////////////////////////////////////
-	// [Insights.MemoryProfiler]
-
-	/** List of search paths to look for symbol files */
-	TArray<FString> SymbolSearchPaths;
-
 	//////////////////////////////////////////////////
 };
diff --git a/Engine/Source/Developer/TraceServices/Private/Common/SymbolHelper.h b/Engine/Source/Developer/TraceServices/Private/Common/SymbolHelper.h
new file mode 100644
--- /dev/null
+++ b/Engine/Source/Developer/TraceServices/Private/Common/SymbolHelper.h
@@ -0,0 +1,58 @@
+﻿// Copyright Epic Games, Inc. All Rights Reserved.
+
+#pragma once
+
+#include "Internationalization/Text.h"
+#include "Misc/ConfigContext.h"
+
+#define LOCTEXT_NAMESPACE "UE::Insights::FSymbolHelper"
+
+namespace TraceServices
+{
+	class FSymbolHelper
+	{
+	public:
+	
+	#if NO_LOGGING
+    	typedef FNoLoggingCategory FLogCategoryAlias;
+    #else
+    	typedef FLogCategoryBase FLogCategoryAlias;
+    #endif
+
+		/**
+		 * Get the list of symbol search paths declared in the UnrealInsightsSettings.ini configuration file.
+		 * @param LogCategory The category the messages should be emitted to.
+		 * @param OutSearchPaths An array that will hold any found paths.
+		 * @return Array of found search paths.
+		 */
+		static void GetPathsFromConfiguration(const FLogCategoryAlias& LogCategory, TArray<FString>& OutSearchPaths)
+		{
+			FString SettingsIni;
+			TArray<FString> SymbolSearchPaths;
+
+			if (FConfigContext::ReadIntoGConfig().Load(TEXT("UnrealInsightsSettings"), SettingsIni))
+			{
+				GConfig->GetArray(TEXT("Insights.Symbols"), TEXT("SymbolSearchPaths"), SymbolSearchPaths, SettingsIni);
+			}
+
+#if !NO_LOGGING
+			if (SymbolSearchPaths.IsEmpty())
+			{
+				UE_LOG_REF(LogCategory, Log, TEXT("[Insights.Symbols] SymbolSearchPaths not set"));
+			}
+			else
+			{
+				for (const FString& Path : SymbolSearchPaths)
+				{
+					UE_LOG_REF(LogCategory, Log, TEXT("[Insights.Symbols] +SymbolSearchPaths=%s"), *Path);
+				}
+			}
+#endif
+
+			OutSearchPaths += SymbolSearchPaths;
+		}
+	};
+
+} // namespace TraceServices
+
+#undef LOCTEXT_NAMESPACE
diff --git a/Engine/Source/Developer/TraceServices/Private/Model/DbgHelpResolver.cpp b/Engine/Source/Developer/TraceServices/Private/Model/DbgHelpResolver.cpp
--- a/Engine/Source/Developer/TraceServices/Private/Model/DbgHelpResolver.cpp
+++ b/Engine/Source/Developer/TraceServices/Private/Model/DbgHelpResolver.cpp
@@ -2,6 +2,9 @@
 
 #include "DbgHelpResolver.h"
 
+// TraceServices
+#include "Common/SymbolHelper.h"
+
 #if PLATFORM_WINDOWS
 
 #include "Algo/ForEach.h"
@@ -12,7 +15,6 @@
 #include "HAL/RunnableThread.h"
 #include "Logging/LogMacros.h"
 #include "Misc/ConfigCacheIni.h"
-#include "Misc/ConfigContext.h"
 #include "Misc/Guid.h"
 #include "Misc/PathViews.h"
 #include "Misc/ScopeLock.h"
@@ -74,27 +76,8 @@ FDbgHelpResolver::FDbgHelpResolver(IAnalysisSession& InSession, IResolvedSymbolF
 	SplitEnvPaths(NTSymbolPath, ConfigSymbolSearchPaths);
 #endif
 
-	// Paths from configuration
-	FString SettingsIni;
-	if (FConfigContext::ReadIntoGConfig().Load(TEXT("UnrealInsightsSettings"), SettingsIni))
-	{
-		TArray<FString> SymbolSearchPaths;
-		GConfig->GetArray(TEXT("Insights.MemoryProfiler"), TEXT("SymbolSearchPaths"), SymbolSearchPaths, SettingsIni);
-#if !NO_LOGGING
-		if (SymbolSearchPaths.IsEmpty())
-		{
-			UE_LOG(LogDbgHelp, Log, TEXT("[Insights.MemoryProfiler] SymbolSearchPaths not set"));
-		}
-		else
-		{
-			for (const FString& Path : SymbolSearchPaths)
-			{
-				UE_LOG(LogDbgHelp, Log, TEXT("[Insights.MemoryProfiler] +SymbolSearchPaths=%s"), *Path);
-			}
-		}
-#endif
-		ConfigSymbolSearchPaths += SymbolSearchPaths;
-	}
+	// Paths from configuration file
+	FSymbolHelper::GetPathsFromConfiguration(LogDbgHelp, ConfigSymbolSearchPaths);
 
 	Start();
 }
diff --git a/Engine/Source/Developer/TraceServices/Private/Model/PsymResolver.cpp b/Engine/Source/Developer/TraceServices/Private/Model/PsymResolver.cpp
--- a/Engine/Source/Developer/TraceServices/Private/Model/PsymResolver.cpp
+++ b/Engine/Source/Developer/TraceServices/Private/Model/PsymResolver.cpp
@@ -1,6 +1,10 @@
 // Copyright Epic Games, Inc. All Rights Reserved.
 
 #include "PsymResolver.h"
+
+// TraceServices
+#include "Common/SymbolHelper.h"
+
 #include "Algo/Sort.h"
 #include "Algo/ForEach.h"
 #include "Containers/Queue.h"
@@ -10,7 +14,6 @@
 #include "HAL/RunnableThread.h"
 #include "Logging/LogMacros.h"
 #include "Misc/ConfigCacheIni.h"
-#include "Misc/ConfigContext.h"
 #include "Misc/Guid.h"
 #include "Misc/PathViews.h"
 #include "Misc/ScopeLock.h"
@@ -83,12 +86,7 @@ FPsymResolver::FPsymResolver(IAnalysisSession& InSession, IResolvedSymbolFilter&
 	// 4. Paths from the user configuration file
 
 	// Paths from configuration
-	FString SettingsIni;
-
-	if (FConfigContext::ReadIntoGConfig().Load(TEXT("UnrealInsightsSettings"), SettingsIni))
-	{
-		GConfig->GetArray(TEXT("Insights.MemoryProfiler"), TEXT("SymbolSearchPaths"), SymbolSearchPaths, SettingsIni);
-	}
+	FSymbolHelper::GetPathsFromConfiguration(LogPsymResolver, SymbolSearchPaths);
 
 	// Paths from environment
 	FString SymbolPathEnvVar =  FPlatformMisc::GetEnvironmentVariable(TEXT("UE_INSIGHTS_SYMBOL_PATH"));
diff --git a/Engine/Source/Developer/TraceServices/Private/Model/SymslibResolver.cpp b/Engine/Source/Developer/TraceServices/Private/Model/SymslibResolver.cpp
--- a/Engine/Source/Developer/TraceServices/Private/Model/SymslibResolver.cpp
+++ b/Engine/Source/Developer/TraceServices/Private/Model/SymslibResolver.cpp
@@ -2,6 +2,9 @@
 
 #include "SymslibResolver.h"
 
+// TraceServices
+#include "Common/SymbolHelper.h"
+
 #if UE_SYMSLIB_AVAILABLE
 
 #include "Algo/ForEach.h"
@@ -15,7 +18,6 @@
 #include "HAL/PlatformTLS.h"
 #include "Logging/LogMacros.h"
 #include "Misc/ConfigCacheIni.h"
-#include "Misc/ConfigContext.h"
 #include "Misc/CString.h"
 #include "Misc/Paths.h"
 #include "Misc/PathViews.h"
@@ -546,27 +548,8 @@ FSymslibResolver::FSymslibResolver(IAnalysisSession& InSession, IResolvedSymbolF
 	SplitEnvPaths(NTSymbolPath, ConfigSymbolSearchPaths);
 #endif
 
-	// Paths from configuration
-	FString SettingsIni;
-	if (FConfigContext::ReadIntoGConfig().Load(TEXT("UnrealInsightsSettings"), SettingsIni))
-	{
-		TArray<FString> SymbolSearchPaths;
-		GConfig->GetArray(TEXT("Insights.MemoryProfiler"), TEXT("SymbolSearchPaths"), SymbolSearchPaths, SettingsIni);
-#if !NO_LOGGING
-		if (SymbolSearchPaths.IsEmpty())
-		{
-			UE_LOG(LogSymslib, Log, TEXT("[Insights.MemoryProfiler] SymbolSearchPaths not set"));
-		}
-		else
-		{
-			for (const FString& Path : SymbolSearchPaths)
-			{
-				UE_LOG(LogSymslib, Log, TEXT("[Insights.MemoryProfiler] +SymbolSearchPaths=%s"), *Path);
-			}
-		}
-#endif
-		ConfigSymbolSearchPaths += SymbolSearchPaths;
-	}
+	// Paths from configuration file
+	FSymbolHelper::GetPathsFromConfiguration(LogSymslib, ConfigSymbolSearchPaths);
 
 	UpdateSessionInfo();
 }
