[Insights]
* Sampling Profiler: Added support for per-frame stats series for sampling timers, for both the Frames view and the Graph tracks.
* Split the FTimingGraphSeries class into separate classes, one for each type of series: frame, counter, timer, per-frame stats timer.

diff --git a/Engine/Source/Developer/TraceInsights/Private/Insights/TimingProfiler/GraphTracks/CounterGraphSeries.cpp b/Engine/Source/Developer/TraceInsights/Private/Insights/TimingProfiler/GraphTracks/CounterGraphSeries.cpp
new file mode 100644
--- /dev/null
+++ b/Engine/Source/Developer/TraceInsights/Private/Insights/TimingProfiler/GraphTracks/CounterGraphSeries.cpp
@@ -0,0 +1,195 @@
+// Copyright Epic Games, Inc. All Rights Reserved.
+
+#include "CounterGraphSeries.h"
+
+// TraceServices
+#include "TraceServices/Model/Counters.h"
+
+// TraceInsightsCore
+#include "InsightsCore/Common/TimeUtils.h"
+
+// TraceInsights
+#include "Insights/InsightsManager.h"
+#include "Insights/TimingProfiler/GraphTracks/TimingGraphTrack.h"
+#include "Insights/ViewModels/GraphTrackBuilder.h"
+#include "Insights/ViewModels/TimingTrackViewport.h"
+
+#define LOCTEXT_NAMESPACE "UE::Insights::FTimingGraphSeries"
+
+namespace UE::Insights::TimingProfiler
+{
+
+INSIGHTS_IMPLEMENT_RTTI(FCounterGraphSeries)
+
+////////////////////////////////////////////////////////////////////////////////////////////////////
+
+void FCounterGraphSeries::InitFromProvider()
+{
+	TSharedPtr<const TraceServices::IAnalysisSession> Session = UE::Insights::FInsightsManager::Get()->GetSession();
+	if (Session.IsValid())
+	{
+		TraceServices::FAnalysisSessionReadScope SessionReadScope(*Session.Get());
+		const TraceServices::ICounterProvider& CountersProvider = TraceServices::ReadCounterProvider(*Session.Get());
+		if (CounterId < CountersProvider.GetCounterCount())
+		{
+			CountersProvider.ReadCounter(CounterId,
+				[this](const TraceServices::ICounter& Counter)
+				{
+					const TCHAR* CounterName = Counter.GetName();
+					if (CounterName)
+					{
+						SetName(CounterName);
+					}
+					//bIsTime = (Counter.GetDisplayHint() == TraceServices::CounterDisplayHint_Time);
+					bIsMemory = (Counter.GetDisplayHint() == TraceServices::CounterDisplayHint_Memory);
+					bIsFloatingPoint = Counter.IsFloatingPoint();
+				});
+		}
+	}
+}
+
+////////////////////////////////////////////////////////////////////////////////////////////////////
+
+FString FCounterGraphSeries::FormatValue(double Value) const
+{
+	if (bIsTime)
+	{
+		return FormatTimeAuto(Value);
+	}
+	else if (bIsMemory)
+	{
+		const int64 MemValue = static_cast<int64>(Value);
+		if (MemValue > 0)
+		{
+			if (MemValue < 1024)
+			{
+				return FString::Printf(TEXT("%s bytes"), *FText::AsNumber(MemValue).ToString());
+			}
+			else
+			{
+				FNumberFormattingOptions FormattingOptions;
+				FormattingOptions.MaximumFractionalDigits = 2;
+				return FString::Printf(TEXT("%s (%s bytes)"), *FText::AsMemory(MemValue, &FormattingOptions).ToString(), *FText::AsNumber(MemValue).ToString());
+			}
+		}
+		else if (MemValue == 0)
+		{
+			return TEXT("0");
+		}
+		else
+		{
+			if (-MemValue < 1024)
+			{
+				return FString::Printf(TEXT("-%s bytes"), *FText::AsNumber(-MemValue).ToString());
+			}
+			else
+			{
+				FNumberFormattingOptions FormattingOptions;
+				FormattingOptions.MaximumFractionalDigits = 2;
+				return FString::Printf(TEXT("-%s (-%s bytes)"), *FText::AsMemory(-MemValue, &FormattingOptions).ToString(), *FText::AsNumber(-MemValue).ToString());
+			}
+		}
+	}
+	else if (bIsFloatingPoint)
+	{
+		return FString::Printf(TEXT("%g"), Value);
+	}
+	else
+	{
+		const int64 Int64Value = static_cast<int64>(Value);
+		return FText::AsNumber(Int64Value).ToString();
+	}
+}
+
+////////////////////////////////////////////////////////////////////////////////////////////////////
+
+void FCounterGraphSeries::Update(FTimingGraphTrack& GraphTrack, const FTimingTrackViewport& Viewport)
+{
+	FGraphTrackBuilder Builder(GraphTrack, *this, Viewport);
+
+	TSharedPtr<const TraceServices::IAnalysisSession> Session = UE::Insights::FInsightsManager::Get()->GetSession();
+	if (!Session.IsValid())
+	{
+		return;
+	}
+
+	TraceServices::FAnalysisSessionReadScope SessionReadScope(*Session.Get());
+
+	const TraceServices::ICounterProvider& CounterProvider = TraceServices::ReadCounterProvider(*Session.Get());
+	CounterProvider.ReadCounter(CounterId,
+		[this, &Viewport, &Builder, &GraphTrack]
+		(const TraceServices::ICounter& Counter)
+		{
+			const float TopY = 4.0f;
+			const float BottomY = GraphTrack.GetHeight() - 4.0f;
+
+			if (IsAutoZoomEnabled() && TopY < BottomY)
+			{
+				double MinValue =  std::numeric_limits<double>::infinity();
+				double MaxValue = -std::numeric_limits<double>::infinity();
+
+				if (Counter.IsFloatingPoint())
+				{
+					Counter.EnumerateFloatValues(Viewport.GetStartTime(), Viewport.GetEndTime(), true,
+						[&Builder, &MinValue, &MaxValue]
+						(double Time, double Value)
+						{
+							if (Value < MinValue)
+							{
+								MinValue = Value;
+							}
+							if (Value > MaxValue)
+							{
+								MaxValue = Value;
+							}
+						});
+				}
+				else
+				{
+					Counter.EnumerateValues(Viewport.GetStartTime(), Viewport.GetEndTime(), true,
+						[&Builder, &MinValue, &MaxValue]
+						(double Time, int64 IntValue)
+						{
+							const double Value = static_cast<double>(IntValue);
+							if (Value < MinValue)
+							{
+								MinValue = Value;
+							}
+							if (Value > MaxValue)
+							{
+								MaxValue = Value;
+							}
+						});
+				}
+
+				UpdateAutoZoom(TopY, BottomY, MinValue, MaxValue);
+			}
+
+			if (Counter.IsFloatingPoint())
+			{
+				Counter.EnumerateFloatValues(Viewport.GetStartTime(), Viewport.GetEndTime(), true,
+					[&Builder]
+					(double Time, double Value)
+					{
+						//TODO: add a "value unit converter"
+						Builder.AddEvent(Time, 0.0, Value);
+					});
+			}
+			else
+			{
+				Counter.EnumerateValues(Viewport.GetStartTime(), Viewport.GetEndTime(), true,
+					[&Builder]
+					(double Time, int64 IntValue)
+					{
+						//TODO: add a "value unit converter"
+						Builder.AddEvent(Time, 0.0, static_cast<double>(IntValue));
+					});
+			}
+		});
+}
+
+////////////////////////////////////////////////////////////////////////////////////////////////////
+
+} // namespace UE::Insights::TimingProfiler
+
+#undef LOCTEXT_NAMESPACE
diff --git a/Engine/Source/Developer/TraceInsights/Private/Insights/TimingProfiler/GraphTracks/CounterGraphSeries.h b/Engine/Source/Developer/TraceInsights/Private/Insights/TimingProfiler/GraphTracks/CounterGraphSeries.h
new file mode 100644
--- /dev/null
+++ b/Engine/Source/Developer/TraceInsights/Private/Insights/TimingProfiler/GraphTracks/CounterGraphSeries.h
@@ -0,0 +1,34 @@
+// Copyright Epic Games, Inc. All Rights Reserved.
+
+#pragma once
+
+#include "CoreTypes.h"
+
+// TraceInsights
+#include "Insights/TimingProfiler/GraphTracks/TimingGraphSeries.h"
+
+namespace UE::Insights::TimingProfiler
+{
+
+class FCounterGraphSeries : public FTimingGraphSeries
+{
+	INSIGHTS_DECLARE_RTTI(FCounterGraphSeries, FTimingGraphSeries)
+
+public:
+	FCounterGraphSeries(uint32 InCounterId) : CounterId(InCounterId) {}
+	virtual ~FCounterGraphSeries() {}
+
+	void InitFromProvider();
+	virtual FString FormatValue(double Value) const override;
+	uint32 GetCounterId() const { return CounterId; }
+	virtual bool IsTimeUnit() const { return bIsTime; }
+	virtual void Update(FTimingGraphTrack& GraphTrack, const FTimingTrackViewport& Viewport) override;
+
+private:
+	uint32 CounterId = 0;
+	bool bIsTime = false; // the unit for values is [second]
+	bool bIsMemory = false; // the unit for value is [byte]
+	bool bIsFloatingPoint = false; // for stats counters
+};
+
+} // namespace UE::Insights::TimingProfiler
diff --git a/Engine/Source/Developer/TraceInsights/Private/Insights/TimingProfiler/GraphTracks/FrameGraphSeries.cpp b/Engine/Source/Developer/TraceInsights/Private/Insights/TimingProfiler/GraphTracks/FrameGraphSeries.cpp
new file mode 100644
--- /dev/null
+++ b/Engine/Source/Developer/TraceInsights/Private/Insights/TimingProfiler/GraphTracks/FrameGraphSeries.cpp
@@ -0,0 +1,112 @@
+// Copyright Epic Games, Inc. All Rights Reserved.
+
+#include "FrameGraphSeries.h"
+
+// TraceServices
+#include "TraceServices/Model/Frames.h"
+
+// TraceInsightsCore
+#include "InsightsCore/Common/TimeUtils.h"
+
+// TraceInsights
+#include "Insights/InsightsManager.h"
+#include "Insights/TimingProfiler/GraphTracks/TimingGraphTrack.h"
+#include "Insights/ViewModels/GraphTrackBuilder.h"
+#include "Insights/ViewModels/TimingTrackViewport.h"
+
+//#include <limits>
+
+#define LOCTEXT_NAMESPACE "UE::Insights::FTimingGraphSeries"
+
+namespace UE::Insights::TimingProfiler
+{
+
+INSIGHTS_IMPLEMENT_RTTI(FFrameGraphSeries)
+
+////////////////////////////////////////////////////////////////////////////////////////////////////
+
+FString FFrameGraphSeries::FormatValue(double Value) const
+{
+	if (Value != 0.0)
+	{
+		return FString::Printf(TEXT("%s (%g fps)"), *FormatTimeAuto(Value), 1.0 / Value);
+	}
+	else
+	{
+		return FString(TEXT("0"));
+	}
+}
+
+////////////////////////////////////////////////////////////////////////////////////////////////////
+
+TSharedRef<FFrameGraphSeries> CreateGameFrameGraphSeries(const FGraphValueViewport& SharedValueViewport)
+{
+	TSharedRef<FFrameGraphSeries> NewSeries = MakeShared<FFrameGraphSeries>(TraceFrameType_Game);
+	FFrameGraphSeries& Series = *NewSeries;
+
+	Series.SetName(TEXT("Game Frames"));
+	Series.SetDescription(TEXT("Duration of Game frames"));
+	Series.SetColor(FLinearColor(0.3f, 0.3f, 1.0f, 1.0f), FLinearColor(1.0f, 1.0f, 1.0f, 1.0f));
+	Series.SetBaselineY(SharedValueViewport.GetBaselineY());
+	Series.SetScaleY(SharedValueViewport.GetScaleY());
+	Series.EnableSharedViewport();
+
+	return NewSeries;
+}
+
+////////////////////////////////////////////////////////////////////////////////////////////////////
+
+TSharedRef<FFrameGraphSeries> CreateRenderingFrameGraphSeries(const FGraphValueViewport& SharedValueViewport)
+{
+	TSharedRef<FFrameGraphSeries> NewSeries = MakeShared<FFrameGraphSeries>(TraceFrameType_Rendering);
+	FFrameGraphSeries& Series = *NewSeries;
+
+	Series.SetName(TEXT("Rendering Frames"));
+	Series.SetDescription(TEXT("Duration of Rendering frames"));
+	Series.SetColor(FLinearColor(1.0f, 0.3f, 0.3f, 1.0f), FLinearColor(1.0f, 1.0f, 1.0f, 1.0f));
+	Series.SetBaselineY(SharedValueViewport.GetBaselineY());
+	Series.SetScaleY(SharedValueViewport.GetScaleY());
+	Series.EnableSharedViewport();
+
+	return NewSeries;
+}
+
+////////////////////////////////////////////////////////////////////////////////////////////////////
+
+void FFrameGraphSeries::Update(FTimingGraphTrack& GraphTrack, const FTimingTrackViewport& Viewport)
+{
+	FGraphTrackBuilder Builder(GraphTrack, *this, Viewport);
+
+	TSharedPtr<const TraceServices::IAnalysisSession> Session = UE::Insights::FInsightsManager::Get()->GetSession();
+	if (!Session.IsValid())
+	{
+		return;
+	}
+
+	TraceServices::FAnalysisSessionReadScope SessionReadScope(*Session.Get());
+
+	const TraceServices::IFrameProvider& FramesProvider = ReadFrameProvider(*Session.Get());
+
+	const TArray64<double>& FrameStartTimes = FramesProvider.GetFrameStartTimes(FrameType);
+
+	const int64 StartLowerBound = Algo::LowerBound(FrameStartTimes, Viewport.GetStartTime());
+	const uint64 StartIndex = (StartLowerBound > 1) ? StartLowerBound - 2 : 0;
+
+	const int64 EndLowerBound = Algo::LowerBound(FrameStartTimes, Viewport.GetEndTime());
+	const uint64 EndIndex = EndLowerBound + 1;
+
+	FramesProvider.EnumerateFrames(FrameType, StartIndex, EndIndex,
+		[&Builder]
+		(const TraceServices::FFrame& Frame)
+		{
+			//TODO: add a "frame converter" (i.e. to fps, milliseconds or seconds)
+			const double Duration = Frame.EndTime - Frame.StartTime;
+			Builder.AddEvent(Frame.StartTime, Duration, Duration);
+		});
+}
+
+////////////////////////////////////////////////////////////////////////////////////////////////////
+
+} // namespace UE::Insights::TimingProfiler
+
+#undef LOCTEXT_NAMESPACE
diff --git a/Engine/Source/Developer/TraceInsights/Private/Insights/TimingProfiler/GraphTracks/FrameGraphSeries.h b/Engine/Source/Developer/TraceInsights/Private/Insights/TimingProfiler/GraphTracks/FrameGraphSeries.h
new file mode 100644
--- /dev/null
+++ b/Engine/Source/Developer/TraceInsights/Private/Insights/TimingProfiler/GraphTracks/FrameGraphSeries.h
@@ -0,0 +1,38 @@
+// Copyright Epic Games, Inc. All Rights Reserved.
+
+#pragma once
+
+#include "CoreTypes.h"
+
+#include "ProfilingDebugging/MiscTrace.h"
+
+// TraceInsights
+#include "Insights/TimingProfiler/GraphTracks/TimingGraphSeries.h"
+
+namespace UE::Insights::TimingProfiler
+{
+
+class FTimingGraphTrack;
+
+class FFrameGraphSeries : public FTimingGraphSeries
+{
+	INSIGHTS_DECLARE_RTTI(FFrameGraphSeries, FTimingGraphSeries)
+
+public:
+	FFrameGraphSeries(ETraceFrameType InFrameType) : FrameType(InFrameType) {}
+	virtual ~FFrameGraphSeries() {}
+
+	ETraceFrameType GetFrameType() const { return FrameType; }
+	virtual bool IsTimeUnit() const override { return true; }
+
+	virtual FString FormatValue(double Value) const override;
+	virtual void Update(FTimingGraphTrack& GraphTrack, const FTimingTrackViewport& Viewport) override;
+
+private:
+	ETraceFrameType FrameType = ETraceFrameType::TraceFrameType_Game;
+};
+
+TSharedRef<FFrameGraphSeries> CreateGameFrameGraphSeries(const FGraphValueViewport& SharedValueViewport);
+TSharedRef<FFrameGraphSeries> CreateRenderingFrameGraphSeries(const FGraphValueViewport& SharedValueViewport);
+
+} // namespace UE::Insights::TimingProfiler
diff --git a/Engine/Source/Developer/TraceInsights/Private/Insights/TimingProfiler/GraphTracks/FrameStatsTimerGraphSeries.cpp b/Engine/Source/Developer/TraceInsights/Private/Insights/TimingProfiler/GraphTracks/FrameStatsTimerGraphSeries.cpp
new file mode 100644
--- /dev/null
+++ b/Engine/Source/Developer/TraceInsights/Private/Insights/TimingProfiler/GraphTracks/FrameStatsTimerGraphSeries.cpp
@@ -0,0 +1,136 @@
+// Copyright Epic Games, Inc. All Rights Reserved.
+
+#include "FrameStatsTimerGraphSeries.h"
+
+#include "Algo/BinarySearch.h"
+
+// TraceServices
+#include "Common/ProviderLock.h"
+#include "TraceServices/Model/Frames.h"
+#include "TraceServices/Model/StackSamples.h"
+#include "TraceServices/Model/TimingProfiler.h"
+
+// TraceInsightsCore
+#include "InsightsCore/Common/TimeUtils.h"
+
+// TraceInsights
+#include "Insights/InsightsManager.h"
+#include "Insights/TimingProfiler/GraphTracks/TimingGraphTrack.h"
+#include "Insights/ViewModels/GraphTrackBuilder.h"
+#include "Insights/ViewModels/TimingTrackViewport.h"
+
+#define LOCTEXT_NAMESPACE "UE::Insights::FTimingGraphSeries"
+
+namespace UE::Insights::TimingProfiler
+{
+
+INSIGHTS_IMPLEMENT_RTTI(FFrameStatsTimerGraphSeries)
+
+////////////////////////////////////////////////////////////////////////////////////////////////////
+
+FString FFrameStatsTimerGraphSeries::FormatValue(double Value) const
+{
+	return FormatTimeAuto(Value);
+}
+
+////////////////////////////////////////////////////////////////////////////////////////////////////
+
+void FFrameStatsTimerGraphSeries::Update(FTimingGraphTrack& GraphTrack, const FTimingTrackViewport& Viewport)
+{
+	FGraphTrackBuilder Builder(GraphTrack, *this, Viewport);
+
+	TSharedPtr<const TraceServices::IAnalysisSession> Session = UE::Insights::FInsightsManager::Get()->GetSession();
+	if (!Session.IsValid())
+	{
+		return;
+	}
+
+	TSet<uint32> VisibleTimelines;
+	GraphTrack.GetVisibleTimelineIndexes(VisibleTimelines);
+
+	TSet<uint32> VisibleCpuSamplingThreads;
+	GraphTrack.GetVisibleCpuSamplingThreads(VisibleCpuSamplingThreads);
+
+	double SessionDuration = 0.0;
+	{
+		TraceServices::FAnalysisSessionReadScope SessionReadScope(*Session.Get());
+		SessionDuration = Session->GetDurationSeconds();
+	}
+
+	if (CachedSessionDuration != SessionDuration ||
+		CachedTimelineCount != VisibleTimelines.Num() ||
+		CachedCpuSamplingTimelineCount != VisibleCpuSamplingThreads.Num())
+	{
+		CachedSessionDuration = SessionDuration;
+		CachedTimelineCount = VisibleTimelines.Num();
+		CachedCpuSamplingTimelineCount = VisibleCpuSamplingThreads.Num();
+		FrameStatsCachedEvents.Empty();
+
+		TraceServices::FAnalysisSessionReadScope SessionReadScope(*Session.Get());
+
+		const TraceServices::IFrameProvider& FramesProvider = ReadFrameProvider(*Session.Get());
+
+		uint64 FrameCount = FramesProvider.GetFrameCount(FrameType);
+		if (FrameCount == 0)
+		{
+			return;
+		}
+
+		FramesProvider.EnumerateFrames(FrameType, 0ull, FrameCount,
+			[this]
+			(const TraceServices::FFrame& Frame)
+			{
+				FFrameStatsCachedEvent Event;
+				Event.FrameStartTime = Frame.StartTime;
+				Event.FrameEndTime = Frame.EndTime;
+				Event.Duration.store(0.0f);
+				FrameStatsCachedEvents.Add(Event);
+			});
+
+		// CPU & GPU timelines
+		if (VisibleTimelines.Num() > 0)
+		{
+			FFrameStatsHelper::ComputeFrameStatsForTimer(FrameStatsCachedEvents, TimerId, VisibleTimelines);
+		}
+
+		// CPU Sampling timelines
+		if (VisibleCpuSamplingThreads.Num() > 0)
+		{
+			const TraceServices::IStackSamplesProvider* StackSamplesProvider = TraceServices::ReadStackSamplesProvider(*Session.Get());
+			if (StackSamplesProvider)
+			{
+				for (uint32 ThreadId : VisibleCpuSamplingThreads)
+				{
+					TraceServices::FProviderReadScopeLock _(*StackSamplesProvider);
+					const TraceServices::ITimingProfilerProvider::Timeline& Timeline = *StackSamplesProvider->GetTimeline(ThreadId);
+					FFrameStatsHelper::ComputeFrameStatsForTimer(FrameStatsCachedEvents, TimerId, Timeline);
+				}
+			}
+		}
+	}
+
+	if (FrameStatsCachedEvents.Num() > 0)
+	{
+		int32 StartIndex = Algo::UpperBoundBy(FrameStatsCachedEvents, Viewport.GetStartTime(), &FFrameStatsCachedEvent::FrameStartTime);
+		if (StartIndex > 0)
+		{
+			StartIndex--;
+		}
+		int32 EndIndex = Algo::UpperBoundBy(FrameStatsCachedEvents, Viewport.GetEndTime(), &FFrameStatsCachedEvent::FrameStartTime);
+		if (EndIndex < FrameStatsCachedEvents.Num())
+		{
+			EndIndex++;
+		}
+		for (int32 Index = StartIndex; Index < EndIndex; ++Index)
+		{
+			const FFrameStatsCachedEvent& Entry = FrameStatsCachedEvents[Index];
+			Builder.AddEvent(Entry.FrameStartTime, Entry.Duration.load(), Entry.Duration.load());
+		}
+	}
+}
+
+////////////////////////////////////////////////////////////////////////////////////////////////////
+
+} // namespace UE::Insights::TimingProfiler
+
+#undef LOCTEXT_NAMESPACE
diff --git a/Engine/Source/Developer/TraceInsights/Private/Insights/TimingProfiler/GraphTracks/FrameStatsTimerGraphSeries.h b/Engine/Source/Developer/TraceInsights/Private/Insights/TimingProfiler/GraphTracks/FrameStatsTimerGraphSeries.h
new file mode 100644
--- /dev/null
+++ b/Engine/Source/Developer/TraceInsights/Private/Insights/TimingProfiler/GraphTracks/FrameStatsTimerGraphSeries.h
@@ -0,0 +1,39 @@
+// Copyright Epic Games, Inc. All Rights Reserved.
+
+#pragma once
+
+#include "CoreTypes.h"
+
+// TraceInsights
+#include "Insights/TimingProfiler/GraphTracks/TimingGraphSeries.h"
+#include "Insights/TimingProfiler/ViewModels/FrameStatsHelper.h"
+
+namespace UE::Insights::TimingProfiler
+{
+
+class FFrameStatsTimerGraphSeries : public FTimingGraphSeries
+{
+	INSIGHTS_DECLARE_RTTI(FFrameStatsTimerGraphSeries, FTimingGraphSeries)
+
+public:
+	FFrameStatsTimerGraphSeries(uint32 InTimerId, ETraceFrameType InFrameType) : TimerId(InTimerId), FrameType(InFrameType) {}
+	virtual ~FFrameStatsTimerGraphSeries() {}
+
+	uint32 GetTimerId() const { return TimerId; }
+	ETraceFrameType GetFrameType() const { return FrameType; }
+	virtual bool IsTimer(uint32 InTimerId) const override { return TimerId == InTimerId; }
+	virtual bool IsTimeUnit() const override { return true; }
+
+	virtual FString FormatValue(double Value) const override;
+	virtual void Update(FTimingGraphTrack& GraphTrack, const FTimingTrackViewport& Viewport) override;
+
+private:
+	uint32 TimerId = 0;
+	ETraceFrameType FrameType = ETraceFrameType::TraceFrameType_Game;
+	double CachedSessionDuration = 0.0;
+	uint32 CachedTimelineCount = 0;
+	uint32 CachedCpuSamplingTimelineCount = 0;
+	TArray<FFrameStatsCachedEvent> FrameStatsCachedEvents;
+};
+
+} // namespace UE::Insights::TimingProfiler
diff --git a/Engine/Source/Developer/TraceInsights/Private/Insights/TimingProfiler/GraphTracks/TimerGraphSeries.cpp b/Engine/Source/Developer/TraceInsights/Private/Insights/TimingProfiler/GraphTracks/TimerGraphSeries.cpp
new file mode 100644
--- /dev/null
+++ b/Engine/Source/Developer/TraceInsights/Private/Insights/TimingProfiler/GraphTracks/TimerGraphSeries.cpp
@@ -0,0 +1,187 @@
+// Copyright Epic Games, Inc. All Rights Reserved.
+
+#include "TimerGraphSeries.h"
+
+#include "Algo/BinarySearch.h"
+
+// TraceServices
+#include "Common/ProviderLock.h"
+#include "TraceServices/Model/StackSamples.h"
+#include "TraceServices/Model/TimingProfiler.h"
+
+// TraceInsightsCore
+#include "InsightsCore/Common/TimeUtils.h"
+
+// TraceInsights
+#include "Insights/InsightsManager.h"
+#include "Insights/TimingProfiler/GraphTracks/TimingGraphTrack.h"
+#include "Insights/ViewModels/GraphTrackBuilder.h"
+#include "Insights/ViewModels/TimingTrackViewport.h"
+
+#define LOCTEXT_NAMESPACE "UE::Insights::FTimingGraphSeries"
+
+namespace UE::Insights::TimingProfiler
+{
+
+INSIGHTS_IMPLEMENT_RTTI(FTimerGraphSeries)
+
+////////////////////////////////////////////////////////////////////////////////////////////////////
+
+FString FTimerGraphSeries::FormatValue(double Value) const
+{
+	return FormatTimeAuto(Value);
+}
+
+////////////////////////////////////////////////////////////////////////////////////////////////////
+
+void FTimerGraphSeries::Update(FTimingGraphTrack& GraphTrack, const FTimingTrackViewport& Viewport)
+{
+	FGraphTrackBuilder Builder(GraphTrack, *this, Viewport);
+
+	TSharedPtr<const TraceServices::IAnalysisSession> Session = UE::Insights::FInsightsManager::Get()->GetSession();
+	if (!Session.IsValid())
+	{
+		return;
+	}
+
+	TSet<uint32> VisibleTimelines;
+	GraphTrack.GetVisibleTimelineIndexes(VisibleTimelines);
+
+	TSet<uint32> VisibleCpuSamplingThreads;
+	GraphTrack.GetVisibleCpuSamplingThreads(VisibleCpuSamplingThreads);
+
+	double SessionDuration = 0.0;
+	{
+		TraceServices::FAnalysisSessionReadScope SessionReadScope(*Session.Get());
+		SessionDuration = Session->GetDurationSeconds();
+	}
+
+	if (CachedSessionDuration != SessionDuration ||
+		CachedTimelineCount != VisibleTimelines.Num() ||
+		CachedCpuSamplingTimelineCount != VisibleCpuSamplingThreads.Num())
+	{
+		CachedSessionDuration = SessionDuration;
+		CachedTimelineCount = VisibleTimelines.Num();
+		CachedCpuSamplingTimelineCount = VisibleCpuSamplingThreads.Num();
+		CachedEvents.Empty();
+
+		const TraceServices::ITimingProfilerProvider* TimingProfilerProvider = TraceServices::ReadTimingProfilerProvider(*Session.Get());
+		if (!TimingProfilerProvider)
+		{
+			return;
+		}
+
+		TraceServices::FAnalysisSessionReadScope SessionReadScope(*Session.Get());
+
+		const TraceServices::ITimingProfilerTimerReader* TimerReader = nullptr;
+		TimingProfilerProvider->ReadTimers([&TimerReader](const TraceServices::ITimingProfilerTimerReader& Out) { TimerReader = &Out; });
+		if (!TimerReader)
+		{
+			return;
+		}
+
+		uint32 NumTimelinesContainingEvent = 0;
+
+		auto TimelineCallback =
+			[this, SessionDuration, &GraphTrack, TimerReader, &Viewport, &NumTimelinesContainingEvent]
+			(const TraceServices::ITimingProfilerProvider::Timeline& Timeline)
+			{
+				TArray<TArray<FSimpleTimingEvent>> Events;
+
+				TraceServices::ITimingProfilerTimeline::EnumerateAsyncParams Params;
+				Params.IntervalStart = 0;
+				Params.IntervalEnd = SessionDuration;
+				Params.Resolution = 0.0;
+				Params.SetupCallback =
+					[&Events]
+					(uint32 NumTasks)
+					{
+						Events.AddDefaulted(NumTasks);
+					};
+				Params.EventRangeCallback =
+					[this, &Events, TimerReader]
+					(double StartTime, double EndTime, uint32 Depth, const TraceServices::FTimingProfilerEvent& Event, uint32 TaskIndex)
+					{
+						const TraceServices::FTimingProfilerTimer* Timer = TimerReader->GetTimer(Event.TimerIndex);
+						if (ensure(Timer != nullptr))
+						{
+							if (Timer->Id == TimerId)
+							{
+								const double Duration = EndTime - StartTime;
+								Events[TaskIndex].Add({ StartTime, Duration });
+							}
+						}
+						return TraceServices::EEventEnumerate::Continue;
+					};
+				Timeline.EnumerateEventsDownSampledAsync(Params);
+
+				int32 NumOfCachedEvents = CachedEvents.Num();
+				for (auto& Array : Events) //-V1078
+				{
+					for (auto& Event : Array) //-V1078
+					{
+						CachedEvents.Add(Event);
+					}
+					Array.Empty();
+				}
+
+				if (NumOfCachedEvents != CachedEvents.Num())
+				{
+					++NumTimelinesContainingEvent;
+				}
+			};
+
+		// CPU & GPU timelines
+		for (uint32 TimelineIndex : VisibleTimelines)
+		{
+			TimingProfilerProvider->ReadTimeline(TimelineIndex, TimelineCallback);
+		}
+
+		// CPU Sampling timelines
+		if (VisibleCpuSamplingThreads.Num() > 0)
+		{
+			const TraceServices::IStackSamplesProvider* StackSamplesProvider = TraceServices::ReadStackSamplesProvider(*Session.Get());
+			if (StackSamplesProvider)
+			{
+				for (uint32 ThreadId : VisibleCpuSamplingThreads)
+				{
+					TraceServices::FProviderReadScopeLock _(*StackSamplesProvider);
+					const TraceServices::ITimingProfilerProvider::Timeline& Timeline = *StackSamplesProvider->GetTimeline(ThreadId);
+					TimelineCallback(Timeline);
+				}
+			}
+		}
+
+		// If events come from multiple timelines, we have to sort the whole thing.
+		// If they come from a single timeline, they are already sorted.
+		if (NumTimelinesContainingEvent > 1)
+		{
+			CachedEvents.Sort(&CompareEventsByStartTime);
+		}
+	}
+
+	if (CachedEvents.Num() > 0)
+	{
+		int32 StartIndex = Algo::UpperBoundBy(CachedEvents, Viewport.GetStartTime(), &FSimpleTimingEvent::StartTime);
+		if (StartIndex > 0)
+		{
+			StartIndex--;
+		}
+		int32 EndIndex = Algo::UpperBoundBy(CachedEvents, Viewport.GetEndTime(), &FSimpleTimingEvent::StartTime);
+		if (EndIndex < CachedEvents.Num())
+		{
+			EndIndex++;
+		}
+		for (int32 Index = StartIndex; Index < EndIndex; ++Index)
+		{
+			const FSimpleTimingEvent& Event = CachedEvents[Index];
+			Builder.AddEvent(Event.StartTime, Event.Duration, Event.Duration);
+		}
+	}
+}
+
+////////////////////////////////////////////////////////////////////////////////////////////////////
+
+} // namespace UE::Insights::TimingProfiler
+
+#undef LOCTEXT_NAMESPACE
diff --git a/Engine/Source/Developer/TraceInsights/Private/Insights/TimingProfiler/GraphTracks/TimerGraphSeries.h b/Engine/Source/Developer/TraceInsights/Private/Insights/TimingProfiler/GraphTracks/TimerGraphSeries.h
new file mode 100644
--- /dev/null
+++ b/Engine/Source/Developer/TraceInsights/Private/Insights/TimingProfiler/GraphTracks/TimerGraphSeries.h
@@ -0,0 +1,49 @@
+// Copyright Epic Games, Inc. All Rights Reserved.
+
+#pragma once
+
+#include "CoreTypes.h"
+
+// TraceInsights
+#include "Insights/TimingProfiler/GraphTracks/TimingGraphSeries.h"
+
+namespace UE::Insights::TimingProfiler
+{
+
+class FTimerGraphSeries : public FTimingGraphSeries
+{
+	INSIGHTS_DECLARE_RTTI(FTimerGraphSeries, FTimingGraphSeries)
+
+private:
+	struct FSimpleTimingEvent
+	{
+		double StartTime;
+		double Duration;
+	};
+
+public:
+	FTimerGraphSeries(uint32 InTimerId) : TimerId(InTimerId) {}
+	virtual ~FTimerGraphSeries() {}
+
+	uint32 GetTimerId() const { return TimerId; }
+	virtual bool IsTimer(uint32 InTimerId) const override { return TimerId == InTimerId; }
+	virtual bool IsTimeUnit() const override { return true; }
+
+	virtual FString FormatValue(double Value) const override;
+	virtual void Update(FTimingGraphTrack& GraphTrack, const FTimingTrackViewport& Viewport) override;
+
+private:
+	static bool CompareEventsByStartTime(const FSimpleTimingEvent& EventA, const FSimpleTimingEvent& EventB)
+	{
+		return EventA.StartTime < EventB.StartTime;
+	}
+
+private:
+	uint32 TimerId = 0;
+	double CachedSessionDuration = 0.0;
+	uint32 CachedTimelineCount = 0;
+	uint32 CachedCpuSamplingTimelineCount = 0;
+	TArray<FSimpleTimingEvent> CachedEvents;
+};
+
+} // namespace UE::Insights::TimingProfiler
diff --git a/Engine/Source/Developer/TraceInsights/Private/Insights/TimingProfiler/GraphTracks/TimingGraphSeries.cpp b/Engine/Source/Developer/TraceInsights/Private/Insights/TimingProfiler/GraphTracks/TimingGraphSeries.cpp
new file mode 100644
--- /dev/null
+++ b/Engine/Source/Developer/TraceInsights/Private/Insights/TimingProfiler/GraphTracks/TimingGraphSeries.cpp
@@ -0,0 +1,25 @@
+// Copyright Epic Games, Inc. All Rights Reserved.
+
+#include "TimingGraphSeries.h"
+
+#define LOCTEXT_NAMESPACE "UE::Insights::FTimingGraphSeries"
+
+namespace UE::Insights::TimingProfiler
+{
+
+INSIGHTS_IMPLEMENT_RTTI(FTimingGraphSeries)
+
+////////////////////////////////////////////////////////////////////////////////////////////////////
+
+void FTimingGraphSeries::SetVisibility(bool bOnOff)
+{
+	FGraphSeries::SetVisibility(bOnOff);
+
+	VisibilityChangedDelegate.Broadcast(bOnOff);
+}
+
+////////////////////////////////////////////////////////////////////////////////////////////////////
+
+} // namespace UE::Insights::TimingProfiler
+
+#undef LOCTEXT_NAMESPACE
diff --git a/Engine/Source/Developer/TraceInsights/Private/Insights/TimingProfiler/GraphTracks/TimingGraphSeries.h b/Engine/Source/Developer/TraceInsights/Private/Insights/TimingProfiler/GraphTracks/TimingGraphSeries.h
new file mode 100644
--- /dev/null
+++ b/Engine/Source/Developer/TraceInsights/Private/Insights/TimingProfiler/GraphTracks/TimingGraphSeries.h
@@ -0,0 +1,40 @@
+// Copyright Epic Games, Inc. All Rights Reserved.
+
+#pragma once
+
+#include "CoreTypes.h"
+#include "Delegates/DelegateCombinations.h"
+
+// TraceInsights
+#include "Insights/ViewModels/GraphSeries.h"
+
+class FTimingTrackViewport;
+
+namespace UE::Insights::TimingProfiler
+{
+
+class FTimingGraphTrack;
+
+/** The delegate to be invoked when a series visibility is changed. */
+DECLARE_MULTICAST_DELEGATE_OneParam(FSeriesVisibilityChangedDelegate, bool bOnOff);
+
+class FTimingGraphSeries : public FGraphSeries
+{
+	INSIGHTS_DECLARE_RTTI(FTimingGraphSeries, FGraphSeries)
+
+public:
+	FTimingGraphSeries() {}
+	virtual ~FTimingGraphSeries() {}
+
+	virtual bool IsTimer(uint32 TimerId) const { return false; }
+	virtual bool IsTimeUnit() const { return false; }
+
+	virtual void SetVisibility(bool bOnOff) override;
+	//virtual FString FormatValue(double Value) const override;
+	virtual void Update(FTimingGraphTrack& GraphTrack, const FTimingTrackViewport& Viewport) = 0;
+
+public:
+	FSeriesVisibilityChangedDelegate VisibilityChangedDelegate;
+};
+
+} // namespace UE::Insights::TimingProfiler
diff --git a/Engine/Source/Developer/TraceInsights/Private/Insights/TimingProfiler/GraphTracks/TimingGraphTrack.cpp b/Engine/Source/Developer/TraceInsights/Private/Insights/TimingProfiler/GraphTracks/TimingGraphTrack.cpp
--- a/Engine/Source/Developer/TraceInsights/Private/Insights/TimingProfiler/GraphTracks/TimingGraphTrack.cpp
+++ b/Engine/Source/Developer/TraceInsights/Private/Insights/TimingProfiler/GraphTracks/TimingGraphTrack.cpp
@@ -2,137 +2,28 @@
 
 #include "TimingGraphTrack.h"
 
-#include "Algo/BinarySearch.h"
-
-// TraceServices
-#include "Common/ProviderLock.h"
-#include "TraceServices/Model/Counters.h"
-#include "TraceServices/Model/Frames.h"
-#include "TraceServices/Model/StackSamples.h"
-#include "TraceServices/Model/TimingProfiler.h"
-
 // TraceInsightsCore
 #include "InsightsCore/Common/PaintUtils.h"
 #include "InsightsCore/Common/TimeUtils.h"
 
 // TraceInsights
 #include "Insights/InsightsManager.h"
+#include "Insights/TimingProfiler/GraphTracks/CounterGraphSeries.h"
+#include "Insights/TimingProfiler/GraphTracks/FrameGraphSeries.h"
+#include "Insights/TimingProfiler/GraphTracks/FrameStatsTimerGraphSeries.h"
+#include "Insights/TimingProfiler/GraphTracks/TimerGraphSeries.h"
 #include "Insights/TimingProfiler/TimingProfilerManager.h"
 #include "Insights/TimingProfiler/Tracks/ThreadTimingTrack.h"
-#include "Insights/TimingProfiler/ViewModels/FrameStatsHelper.h"
 #include "Insights/TimingProfiler/Widgets/STimingProfilerWindow.h"
 #include "Insights/ViewModels/AxisViewportDouble.h"
-#include "Insights/ViewModels/GraphTrackBuilder.h"
-#include "Insights/ViewModels/ITimingViewDrawHelper.h"
-#include "Insights/ViewModels/TimingTrackViewport.h"
 #include "Insights/Widgets/STimingView.h"
 
 #include <limits>
 
-#define LOCTEXT_NAMESPACE "GraphTrack"
-
-INSIGHTS_IMPLEMENT_RTTI(FTimingGraphSeries)
-
-////////////////////////////////////////////////////////////////////////////////////////////////////
-// FTimingGraphSeries
-////////////////////////////////////////////////////////////////////////////////////////////////////
+#define LOCTEXT_NAMESPACE "UE::Insights::TimingProfiler::FTimingGraphTrack"
 
-FTimingGraphSeries::FTimingGraphSeries(FTimingGraphSeries::ESeriesType InType)
-	: FGraphSeries()
-	, Type(InType)
-	, TimerId(0)
-	, CachedSessionDuration(0.0)
-	, CachedEvents()
-	, bIsTime(InType == FTimingGraphSeries::ESeriesType::Frame ||
-			  InType == FTimingGraphSeries::ESeriesType::Timer ||
-			  InType == FTimingGraphSeries::ESeriesType::FrameStatsTimer)
-	, bIsMemory(false)
-	, bIsFloatingPoint(false)
+namespace UE::Insights::TimingProfiler
 {
-}
-
-////////////////////////////////////////////////////////////////////////////////////////////////////
-
-FTimingGraphSeries::~FTimingGraphSeries()
-{
-}
-
-////////////////////////////////////////////////////////////////////////////////////////////////////
-
-FString FTimingGraphSeries::FormatValue(double Value) const
-{
-	using namespace UE::Insights;
-
-	switch (Type)
-	{
-	case FTimingGraphSeries::ESeriesType::Frame:
-		return FString::Printf(TEXT("%s (%g fps)"), *FormatTimeAuto(Value), 1.0 / Value);
-
-	case FTimingGraphSeries::ESeriesType::Timer:
-	case FTimingGraphSeries::ESeriesType::FrameStatsTimer:
-		return FormatTimeAuto(Value);
-
-	case FTimingGraphSeries::ESeriesType::StatsCounter:
-		if (bIsTime)
-		{
-			return FormatTimeAuto(Value);
-		}
-		else if (bIsMemory)
-		{
-			const int64 MemValue = static_cast<int64>(Value);
-			if (MemValue > 0)
-			{
-				if (MemValue < 1024)
-				{
-					return FString::Printf(TEXT("%s bytes"), *FText::AsNumber(MemValue).ToString());
-				}
-				else
-				{
-					FNumberFormattingOptions FormattingOptions;
-					FormattingOptions.MaximumFractionalDigits = 2;
-					return FString::Printf(TEXT("%s (%s bytes)"), *FText::AsMemory(MemValue, &FormattingOptions).ToString(), *FText::AsNumber(MemValue).ToString());
-				}
-			}
-			else if (MemValue == 0)
-			{
-				return TEXT("0");
-			}
-			else
-			{
-				if (-MemValue < 1024)
-				{
-					return FString::Printf(TEXT("-%s bytes"), *FText::AsNumber(-MemValue).ToString());
-				}
-				else
-				{
-					FNumberFormattingOptions FormattingOptions;
-					FormattingOptions.MaximumFractionalDigits = 2;
-					return FString::Printf(TEXT("-%s (-%s bytes)"), *FText::AsMemory(-MemValue, &FormattingOptions).ToString(), *FText::AsNumber(-MemValue).ToString());
-				}
-			}
-		}
-		else if (bIsFloatingPoint)
-		{
-			return FString::Printf(TEXT("%g"), Value);
-		}
-		else
-		{
-			const int64 Int64Value = static_cast<int64>(Value);
-			return FText::AsNumber(Int64Value).ToString();
-		}
-	}
-
-	return FGraphSeries::FormatValue(Value);
-}
-
-////////////////////////////////////////////////////////////////////////////////////////////////////
-
-void FTimingGraphSeries::SetVisibility(bool bOnOff)
-{
-	FGraphSeries::SetVisibility(bOnOff);
-
-	VisibilityChangedDelegate.Broadcast(bOnOff);
-}
 
 ////////////////////////////////////////////////////////////////////////////////////////////////////
 // FTimingGraphTrack
@@ -142,7 +33,7 @@ INSIGHTS_IMPLEMENT_RTTI(FTimingGraphTrack)
 
 ////////////////////////////////////////////////////////////////////////////////////////////////////
 
-FTimingGraphTrack::FTimingGraphTrack(TSharedPtr<UE::Insights::TimingProfiler::STimingView> InTimingView)
+FTimingGraphTrack::FTimingGraphTrack(TSharedPtr<STimingView> InTimingView)
 	: FGraphTrack()
 	, TimingView(InTimingView)
 {
@@ -164,13 +55,7 @@ FTimingGraphTrack::~FTimingGraphTrack()
 {
 	if (OnTrackVisibilityChangedHandle.IsValid())
 	{
-		TSharedPtr<UE::Insights::TimingProfiler::STimingView> TimingViewPtr = TimingView.Pin();
-		if (TimingView.IsValid())
-		{
-			TimingViewPtr->OnTrackVisibilityChanged().Remove(OnTrackVisibilityChangedHandle);
-			TimingViewPtr->OnTrackAdded().Remove(OnTrackAddedHandle);
-			TimingViewPtr->OnTrackRemoved().Remove(OnTrackRemovedHandle);
-		}
+		UnregisterTimingViewCallbacks();
 	}
 
 	if (GameFrameSeriesVisibilityHandle.IsValid())
@@ -191,28 +76,27 @@ FTimingGraphTrack::~FTimingGraphTrack()
 		}
 	}
 
-	using namespace UE::Insights::TimingProfiler;
-
 	if (bNotifyTimersOnDestruction)
 	{
 		TSharedPtr<STimingProfilerWindow> ProfilerWindow = FTimingProfilerManager::Get()->GetProfilerWindow();
 		if (ProfilerWindow.IsValid())
 		{
-			TArray<TSharedPtr<FTimingGraphSeries>> TimerSeries;
+			TArray<uint32> TimerIds;
 			for (TSharedPtr<FGraphSeries>& Series : AllSeries)
 			{
-				TSharedPtr<FTimingGraphSeries> TimingSeries = StaticCastSharedPtr<FTimingGraphSeries>(Series);
-				if (TimingSeries.IsValid() &&
-					(TimingSeries->Type == FTimingGraphSeries::ESeriesType::Timer ||
-					 TimingSeries->Type == FTimingGraphSeries::ESeriesType::FrameStatsTimer))
+				if (Series->Is<FTimerGraphSeries>())
+				{
+					TimerIds.Add(Series->As<FTimerGraphSeries>().GetTimerId());
+				}
+				else if (Series->Is<FFrameStatsTimerGraphSeries>())
 				{
-					TimerSeries.Add(TimingSeries);
+					TimerIds.Add(Series->As<FFrameStatsTimerGraphSeries>().GetTimerId());
 				}
 			}
 			AllSeries.Reset();
-			for (TSharedPtr<FTimingGraphSeries>& TimingSeries : TimerSeries)
+			for (uint32 TimerId : TimerIds)
 			{
-				ProfilerWindow->OnTimerAddedToGraphsChanged(TimingSeries->TimerId);
+				ProfilerWindow->OnTimerAddedToGraphsChanged(TimerId);
 			}
 		}
 	}
@@ -220,35 +104,65 @@ FTimingGraphTrack::~FTimingGraphTrack()
 
 ////////////////////////////////////////////////////////////////////////////////////////////////////
 
-void FTimingGraphTrack::Update(const ITimingTrackUpdateContext& Context)
+void FTimingGraphTrack::RegisterTimingViewCallbacks()
 {
-	FGraphTrack::Update(Context);
-
-	if (!OnTrackVisibilityChangedHandle.IsValid())
+	TSharedPtr<STimingView> TimingViewPtr = TimingView.Pin();
+	if (TimingViewPtr.IsValid())
 	{
-		TSharedPtr<UE::Insights::TimingProfiler::STimingView> TimingViewPtr = TimingView.Pin();
-		if (TimingViewPtr.IsValid())
-		{
-			auto OnTrackAddedRemovedLamda = [this](const TSharedPtr<const FBaseTimingTrack> Track)
+		auto OnTrackAddedRemovedLambda =
+			[this](const TSharedPtr<const FBaseTimingTrack> Track)
 			{
-				if (Track->Is<UE::Insights::TimingProfiler::FThreadTimingTrack>())
+				if (Track->Is<FThreadTimingTrack>())
 				{
 					// If there are more series than the default frame series.
-					if (this->AllSeries.Num() > ETraceFrameType::TraceFrameType_Count)
+					if (AllSeries.Num() > ETraceFrameType::TraceFrameType_Count)
 					{
-						this->SetDirtyFlag();
+						SetDirtyFlag();
 					}
 				}
 			};
 
-			OnTrackAddedHandle = TimingViewPtr->OnTrackAdded().AddLambda(OnTrackAddedRemovedLamda);
-			OnTrackRemovedHandle = TimingViewPtr->OnTrackRemoved().AddLambda(OnTrackAddedRemovedLamda);
+		OnTrackAddedHandle = TimingViewPtr->OnTrackAdded().AddLambda(OnTrackAddedRemovedLambda);
+		OnTrackRemovedHandle = TimingViewPtr->OnTrackRemoved().AddLambda(OnTrackAddedRemovedLambda);
 
-			OnTrackVisibilityChangedHandle = TimingViewPtr->OnTrackVisibilityChanged().AddLambda([this]()
+		OnTrackVisibilityChangedHandle = TimingViewPtr->OnTrackVisibilityChanged().AddLambda(
+			[this]()
+			{
+				// If there are more series than the default frame series.
+				if (AllSeries.Num() > ETraceFrameType::TraceFrameType_Count)
 				{
-					this->SetDirtyFlag();
-				});
-		}
+					SetDirtyFlag();
+				}
+			});
+	}
+}
+
+////////////////////////////////////////////////////////////////////////////////////////////////////
+
+void FTimingGraphTrack::UnregisterTimingViewCallbacks()
+{
+	TSharedPtr<STimingView> TimingViewPtr = TimingView.Pin();
+	if (TimingView.IsValid())
+	{
+		TimingViewPtr->OnTrackAdded().Remove(OnTrackAddedHandle);
+		TimingViewPtr->OnTrackRemoved().Remove(OnTrackRemovedHandle);
+		TimingViewPtr->OnTrackVisibilityChanged().Remove(OnTrackVisibilityChangedHandle);
+	}
+
+	OnTrackAddedHandle.Reset();
+	OnTrackRemovedHandle.Reset();
+	OnTrackVisibilityChangedHandle.Reset();
+}
+
+////////////////////////////////////////////////////////////////////////////////////////////////////
+
+void FTimingGraphTrack::Update(const ITimingTrackUpdateContext& Context)
+{
+	FGraphTrack::Update(Context);
+
+	if (!OnTrackVisibilityChangedHandle.IsValid())
+	{
+		RegisterTimingViewCallbacks();
 	}
 
 	const bool bIsEntireGraphTrackDirty = IsDirty() || Context.GetViewport().IsHorizontalViewportDirty();
@@ -277,30 +191,14 @@ void FTimingGraphTrack::Update(const ITimingTrackUpdateContext& Context)
 
 		for (TSharedPtr<FGraphSeries>& Series : AllSeries)
 		{
-			if (Series->IsVisible() && (bIsEntireGraphTrackDirty || Series->IsDirty()))
+			if (Series->IsVisible() &&
+				(bIsEntireGraphTrackDirty || Series->IsDirty()) &&
+				Series->Is<FTimingGraphSeries>())
 			{
 				// Clear the flag before updating, because the update itself may further need to set the series as dirty.
 				Series->ClearDirtyFlag();
 
-				TSharedPtr<FTimingGraphSeries> TimingSeries = StaticCastSharedPtr<FTimingGraphSeries>(Series);
-				switch (TimingSeries->Type)
-				{
-				case FTimingGraphSeries::ESeriesType::Frame:
-					UpdateFrameSeries(*TimingSeries, Viewport);
-					break;
-
-				case FTimingGraphSeries::ESeriesType::Timer:
-					UpdateTimerSeries(*TimingSeries, Viewport);
-					break;
-
-				case FTimingGraphSeries::ESeriesType::FrameStatsTimer:
-					UpdateFrameStatsTimerSeries(*TimingSeries, Viewport);
-					break;
-
-				case FTimingGraphSeries::ESeriesType::StatsCounter:
-					UpdateStatsCounterSeries(*TimingSeries, Viewport);
-					break;
-				}
+				Series->As<FTimingGraphSeries>().Update(*this, Viewport);
 			}
 		}
 
@@ -314,99 +212,64 @@ void FTimingGraphTrack::Update(const ITimingTrackUpdateContext& Context)
 
 void FTimingGraphTrack::AddDefaultFrameSeries()
 {
-	const FInsightsSettings& Settings = UE::Insights::FInsightsManager::Get()->GetSettings();
-	TSharedPtr<UE::Insights::TimingProfiler::STimingView> TimingViewPtr = TimingView.Pin();
-
-	TSharedRef<FTimingGraphSeries> GameFramesSeries = MakeShared<FTimingGraphSeries>(FTimingGraphSeries::ESeriesType::Frame);
-	GameFramesSeries->SetName(TEXT("Game Frames"));
-	GameFramesSeries->SetDescription(TEXT("Duration of Game frames"));
-	GameFramesSeries->SetColor(FLinearColor(0.3f, 0.3f, 1.0f, 1.0f), FLinearColor(1.0f, 1.0f, 1.0f, 1.0f));
-	GameFramesSeries->FrameType = TraceFrameType_Game;
-	GameFramesSeries->SetBaselineY(SharedValueViewport.GetBaselineY());
-	GameFramesSeries->SetScaleY(SharedValueViewport.GetScaleY());
-	GameFramesSeries->EnableSharedViewport();
-	if (TimingViewPtr.IsValid() && TimingViewPtr->GetName() == FInsightsManagerTabs::TimingProfilerTabId)
+	TSharedPtr<STimingView> TimingViewPtr = TimingView.Pin();
+
+	TSharedRef<FFrameGraphSeries> GameFramesSeries = CreateGameFrameGraphSeries(SharedValueViewport);
+	TSharedRef<FFrameGraphSeries> RenderingFramesSeries = CreateRenderingFrameGraphSeries(SharedValueViewport);
+
+	if (TimingViewPtr.IsValid() &&
+		TimingViewPtr->GetName() == FInsightsManagerTabs::TimingProfilerTabId)
 	{
+		const FInsightsSettings& Settings = UE::Insights::FInsightsManager::Get()->GetSettings();
+
 		GameFramesSeries->SetVisibility(Settings.GetTimingViewMainGraphShowGameFrames());
-		GameFrameSeriesVisibilityHandle = GameFramesSeries->VisibilityChangedDelegate.AddLambda([](bool bOnOff)
+		GameFrameSeriesVisibilityHandle = GameFramesSeries->VisibilityChangedDelegate.AddLambda(
+			[](bool bOnOff)
 			{
 				FInsightsSettings& Settings = UE::Insights::FInsightsManager::Get()->GetSettings();
 				Settings.SetAndSaveTimingViewMainGraphShowGameFrames(bOnOff);
 			});
-	}
-	AllSeries.Add(GameFramesSeries);
-
-	TSharedRef<FTimingGraphSeries> RenderingFramesSeries = MakeShared<FTimingGraphSeries>(FTimingGraphSeries::ESeriesType::Frame);
-	RenderingFramesSeries->SetName(TEXT("Rendering Frames"));
-	RenderingFramesSeries->SetDescription(TEXT("Duration of Rendering frames"));
-	RenderingFramesSeries->SetColor(FLinearColor(1.0f, 0.3f, 0.3f, 1.0f), FLinearColor(1.0f, 1.0f, 1.0f, 1.0f));
-	RenderingFramesSeries->FrameType = TraceFrameType_Rendering;
-	RenderingFramesSeries->SetBaselineY(SharedValueViewport.GetBaselineY());
-	RenderingFramesSeries->SetScaleY(SharedValueViewport.GetScaleY());
-	RenderingFramesSeries->EnableSharedViewport();
-	if (TimingViewPtr.IsValid() && TimingViewPtr->GetName() == FInsightsManagerTabs::TimingProfilerTabId)
-	{
+
 		RenderingFramesSeries->SetVisibility(Settings.GetTimingViewMainGraphShowRenderingFrames());
-		RenderingFrameSeriesVisibilityHandle = RenderingFramesSeries->VisibilityChangedDelegate.AddLambda([](bool bOnOff)
+		RenderingFrameSeriesVisibilityHandle = RenderingFramesSeries->VisibilityChangedDelegate.AddLambda(
+			[](bool bOnOff)
 			{
 				FInsightsSettings& Settings = UE::Insights::FInsightsManager::Get()->GetSettings();
 				Settings.SetAndSaveTimingViewMainGraphShowRenderingFrames(bOnOff);
 			});
 	}
-	AllSeries.Add(RenderingFramesSeries);
+
+	AllSeries.Add(GameFramesSeries.ToSharedPtr());
+	AllSeries.Add(RenderingFramesSeries.ToSharedPtr());
 }
 
 ////////////////////////////////////////////////////////////////////////////////////////////////////
 
 TSharedPtr<FTimingGraphSeries> FTimingGraphTrack::GetFrameSeries(ETraceFrameType FrameType)
 {
-	TSharedPtr<FGraphSeries>* Ptr = AllSeries.FindByPredicate([FrameType](const TSharedPtr<FGraphSeries>& Series)
+	TSharedPtr<FGraphSeries>* Ptr = AllSeries.FindByPredicate(
+		[FrameType]
+		(const TSharedPtr<FGraphSeries>& Series)
 		{
-			const TSharedPtr<FTimingGraphSeries> TimingSeries = StaticCastSharedPtr<FTimingGraphSeries>(Series);
-			return TimingSeries->Type == FTimingGraphSeries::ESeriesType::Frame && TimingSeries->FrameType == FrameType;
+			return Series->Is<FFrameGraphSeries>() &&
+				   Series->As<FFrameGraphSeries>().GetFrameType() == FrameType;
 		});
 	return (Ptr != nullptr) ? StaticCastSharedPtr<FTimingGraphSeries>(*Ptr) : nullptr;
 }
 
-////////////////////////////////////////////////////////////////////////////////////////////////////
-
-void FTimingGraphTrack::UpdateFrameSeries(FTimingGraphSeries& Series, const FTimingTrackViewport& Viewport)
-{
-	FGraphTrackBuilder Builder(*this, Series, Viewport);
-	TSharedPtr<const TraceServices::IAnalysisSession> Session = UE::Insights::FInsightsManager::Get()->GetSession();
-	if (Session.IsValid())
-	{
-		TraceServices::FAnalysisSessionReadScope SessionReadScope(*Session.Get());
-		const TraceServices::IFrameProvider& FramesProvider = ReadFrameProvider(*Session.Get());
-
-		const TArray64<double>& FrameStartTimes = FramesProvider.GetFrameStartTimes(Series.FrameType);
-
-		const int64 StartLowerBound = Algo::LowerBound(FrameStartTimes, Viewport.GetStartTime());
-		const uint64 StartIndex = (StartLowerBound > 1) ? StartLowerBound - 2 : 0;
-
-		const int64 EndLowerBound = Algo::LowerBound(FrameStartTimes, Viewport.GetEndTime());
-		const uint64 EndIndex = EndLowerBound + 1;
-
-		FramesProvider.EnumerateFrames(Series.FrameType, StartIndex, EndIndex, [&Builder](const TraceServices::FFrame& Frame)
-		{
-			//TODO: add a "frame converter" (i.e. to fps, milliseconds or seconds)
-			const double Duration = Frame.EndTime - Frame.StartTime;
-			Builder.AddEvent(Frame.StartTime, Duration, Duration);
-		});
-	}
-}
-
 ////////////////////////////////////////////////////////////////////////////////////////////////////
 // Timer Series
 ////////////////////////////////////////////////////////////////////////////////////////////////////
 
 TSharedPtr<FTimingGraphSeries> FTimingGraphTrack::GetTimerSeries(uint32 TimerId)
 {
-	TSharedPtr<FGraphSeries>* Ptr = AllSeries.FindByPredicate([TimerId](const TSharedPtr<FGraphSeries>& Series)
-	{
-		const TSharedPtr<FTimingGraphSeries> TimingSeries = StaticCastSharedPtr<FTimingGraphSeries>(Series);
-		return TimingSeries->Type == FTimingGraphSeries::ESeriesType::Timer && TimingSeries->TimerId == TimerId;
-	});
+	TSharedPtr<FGraphSeries>* Ptr = AllSeries.FindByPredicate(
+		[TimerId]
+		(const TSharedPtr<FGraphSeries>& Series)
+		{
+			return Series->Is<FTimerGraphSeries>() &&
+				   Series->As<FTimerGraphSeries>().GetTimerId() == TimerId;
+		});
 	return (Ptr != nullptr) ? StaticCastSharedPtr<FTimingGraphSeries>(*Ptr) : nullptr;
 }
 
@@ -414,155 +277,35 @@ TSharedPtr<FTimingGraphSeries> FTimingGraphTrack::GetTimerSeries(uint32 TimerId)
 
 TSharedPtr<FTimingGraphSeries> FTimingGraphTrack::AddTimerSeries(uint32 TimerId, FLinearColor Color)
 {
-	TSharedRef<FTimingGraphSeries> Series = MakeShared<FTimingGraphSeries>(FTimingGraphSeries::ESeriesType::Timer);
+	TSharedRef<FTimerGraphSeries> NewSeries = MakeShared<FTimerGraphSeries>(TimerId);
+	FTimerGraphSeries& Series = *NewSeries;
 
-	Series->SetName(TEXT("<Timer>"));
-	Series->SetDescription(TEXT("Timer series"));
+	Series.SetName(TEXT("<Timer>"));
+	Series.SetDescription(TEXT("Timer series"));
 
 	const FLinearColor BorderColor(Color.R + 0.4f, Color.G + 0.4f, Color.B + 0.4f, 1.0f);
-	Series->SetColor(Color, BorderColor);
-
-	Series->TimerId = TimerId;
-	//Series->CpuOrGpu = ;
-	//Series->TimelineIndex = ;
+	Series.SetColor(Color, BorderColor);
 
 	// Use shared viewport.
-	Series->SetBaselineY(SharedValueViewport.GetBaselineY());
-	Series->SetScaleY(SharedValueViewport.GetScaleY());
-	Series->EnableSharedViewport();
-
-	Series->CachedSessionDuration = 0.0;
+	Series.SetBaselineY(SharedValueViewport.GetBaselineY());
+	Series.SetScaleY(SharedValueViewport.GetScaleY());
+	Series.EnableSharedViewport();
 
-	AllSeries.Add(Series);
-	return Series;
+	AllSeries.Add(NewSeries);
+	return NewSeries;
 }
 
 ////////////////////////////////////////////////////////////////////////////////////////////////////
 
 void FTimingGraphTrack::RemoveTimerSeries(uint32 TimerId)
 {
-	AllSeries.RemoveAll([TimerId](const TSharedPtr<FGraphSeries>& Series)
-	{
-		const TSharedPtr<FTimingGraphSeries> TimingSeries = StaticCastSharedPtr<FTimingGraphSeries>(Series);
-		return TimingSeries->Type == FTimingGraphSeries::ESeriesType::Timer && TimingSeries->TimerId == TimerId;
-	});
-}
-
-////////////////////////////////////////////////////////////////////////////////////////////////////
-
-void FTimingGraphTrack::UpdateTimerSeries(FTimingGraphSeries& Series, const FTimingTrackViewport& Viewport)
-{
-	FGraphTrackBuilder Builder(*this, Series, Viewport);
-	TSharedPtr<const TraceServices::IAnalysisSession> Session = UE::Insights::FInsightsManager::Get()->GetSession();
-	if (Session.IsValid())
-	{
-		TraceServices::FAnalysisSessionReadScope SessionReadScope(*Session.Get());
-
-		TSet<uint32> Timelines;
-		GetVisibleTimelineIndexes(Timelines);
-		TSet<uint32> VisibleStackThreads;
-		GetVisibleCpuSamplingThreads(VisibleStackThreads);
-
-		const double SessionDuration = Session->GetDurationSeconds();
-		if (Series.CachedSessionDuration != SessionDuration || Series.CachedTimelinesNum != Timelines.Num())
+	AllSeries.RemoveAll(
+		[TimerId]
+		(const TSharedPtr<FGraphSeries>& Series)
 		{
-			Series.CachedSessionDuration = SessionDuration;
-
-			const TraceServices::ITimingProfilerProvider& TimingProfilerProvider = *TraceServices::ReadTimingProfilerProvider(*Session.Get());
-
-			const TraceServices::ITimingProfilerTimerReader* TimerReader;
-			TimingProfilerProvider.ReadTimers([&TimerReader](const TraceServices::ITimingProfilerTimerReader& Out) { TimerReader = &Out; });
-
-			const uint32 TimelineCount = TimingProfilerProvider.GetTimelineCount();
-			uint32 NumTimelinesContainingEvent = 0;
-
-			Series.CachedTimelinesNum = Timelines.Num();
-			Series.CachedEvents.Empty();
-
-			auto TimelineCallback = [SessionDuration, &Series, TimerReader, &Viewport, &NumTimelinesContainingEvent](const TraceServices::ITimingProfilerProvider::Timeline& Timeline)
-			{
-				TArray<TArray<FTimingGraphSeries::FSimpleTimingEvent>> Events;
-				TraceServices::ITimingProfilerTimeline::EnumerateAsyncParams Params;
-				Params.IntervalStart = 0;
-				Params.IntervalEnd = SessionDuration;
-				Params.Resolution = 0.0;
-				Params.SetupCallback = [&Events](uint32 NumTasks)
-					{
-						Events.AddDefaulted(NumTasks);
-					};
-				Params.EventRangeCallback = [&Events, TimerReader, &Series](double StartTime, double EndTime, uint32 Depth, const TraceServices::FTimingProfilerEvent& Event, uint32 TaskIndex)
-					{
-						const TraceServices::FTimingProfilerTimer* Timer = TimerReader->GetTimer(Event.TimerIndex);
-						if (ensure(Timer != nullptr))
-						{
-							if (Timer->Id == Series.TimerId)
-							{
-								const double Duration = EndTime - StartTime;
-								Events[TaskIndex].Add({ StartTime, Duration });
-							}
-						}
-						return TraceServices::EEventEnumerate::Continue;
-					};
-
-				Timeline.EnumerateEventsDownSampledAsync(Params);
-
-				int32 NumOfCachedEvents = Series.CachedEvents.Num();
-				for (auto& Array : Events) //-V1078
-				{
-					for (auto& Event : Array) //-V1078
-					{
-						Series.CachedEvents.Add(Event);
-					}
-
-					Array.Empty();
-				}
-
-				if (NumOfCachedEvents != Series.CachedEvents.Num())
-				{
-					++NumTimelinesContainingEvent;
-				}
-			};
-
-			for (uint32 TimelineIndex : Timelines)
-			{
-				TimingProfilerProvider.ReadTimeline(TimelineIndex, TimelineCallback);
-			}
-
-			const TraceServices::IStackSamplesProvider* StackSamplesProvider = TraceServices::ReadStackSamplesProvider(*Session.Get());
-			if (VisibleStackThreads.Num() > 0 && StackSamplesProvider)
-			{
-				for (uint32 ThreadId : VisibleStackThreads)
-				{
-					TraceServices::FProviderReadScopeLock _(*StackSamplesProvider);
-					const TraceServices::ITimingProfilerProvider::Timeline& Timeline = *StackSamplesProvider->GetTimeline(ThreadId);
-					TimelineCallback(Timeline);
-				}
-			}
-
-			//If events come from multiple timelines, we have to sort the whole thing.
-			//If they come from a single timeline, they are already sorted.
-			if (NumTimelinesContainingEvent > 1)
-			{
-				Series.CachedEvents.Sort(&FTimingGraphSeries::CompareEventsByStartTime);
-			}
-		}
-
-		int32 StartIndex = Algo::UpperBoundBy(Series.CachedEvents, Viewport.GetStartTime(), &FTimingGraphSeries::FSimpleTimingEvent::StartTime);
-		if (StartIndex > 0)
-		{
-			StartIndex--;
-		}
-		int32 EndIndex = Algo::UpperBoundBy(Series.CachedEvents, Viewport.GetEndTime(), &FTimingGraphSeries::FSimpleTimingEvent::StartTime);
-		if (EndIndex < Series.CachedEvents.Num())
-		{
-			EndIndex++;
-		}
-		for (int32 Index = StartIndex; Index < EndIndex; ++Index)
-		{
-			const FTimingGraphSeries::FSimpleTimingEvent& Event = Series.CachedEvents[Index];
-			Builder.AddEvent(Event.StartTime, Event.Duration, Event.Duration);
-		}
-	}
+			return Series->Is<FTimerGraphSeries>() &&
+				   Series->As<FTimerGraphSeries>().GetTimerId() == TimerId;
+		});
 }
 
 ////////////////////////////////////////////////////////////////////////////////////////////////////
@@ -571,12 +314,13 @@ void FTimingGraphTrack::UpdateTimerSeries(FTimingGraphSeries& Series, const FTim
 
 TSharedPtr<FTimingGraphSeries> FTimingGraphTrack::GetFrameStatsTimerSeries(uint32 TimerId, ETraceFrameType FrameType)
 {
-	TSharedPtr<FGraphSeries>* Ptr = AllSeries.FindByPredicate([TimerId, FrameType](const TSharedPtr<FGraphSeries>& Series)
+	TSharedPtr<FGraphSeries>* Ptr = AllSeries.FindByPredicate(
+		[TimerId, FrameType]
+		(const TSharedPtr<FGraphSeries>& Series)
 		{
-			const TSharedPtr<FTimingGraphSeries> TimingSeries = StaticCastSharedPtr<FTimingGraphSeries>(Series);
-			return TimingSeries->Type == FTimingGraphSeries::ESeriesType::FrameStatsTimer &&
-				   TimingSeries->TimerId == TimerId &&
-				   TimingSeries->FrameType == FrameType;
+			return Series->Is<FFrameStatsTimerGraphSeries>() &&
+				   Series->As<FFrameStatsTimerGraphSeries>().GetTimerId() == TimerId &&
+				   Series->As<FFrameStatsTimerGraphSeries>().GetFrameType() == FrameType;
 		});
 	return (Ptr != nullptr) ? StaticCastSharedPtr<FTimingGraphSeries>(*Ptr) : nullptr;
 }
@@ -585,111 +329,51 @@ TSharedPtr<FTimingGraphSeries> FTimingGraphTrack::GetFrameStatsTimerSeries(uint3
 
 TSharedPtr<FTimingGraphSeries> FTimingGraphTrack::AddFrameStatsTimerSeries(uint32 TimerId, ETraceFrameType FrameType, FLinearColor Color)
 {
-	TSharedRef<FTimingGraphSeries> Series = MakeShared<FTimingGraphSeries>(FTimingGraphSeries::ESeriesType::FrameStatsTimer);
+	TSharedRef<FFrameStatsTimerGraphSeries> NewSeries = MakeShared<FFrameStatsTimerGraphSeries>(TimerId, FrameType);
+	FFrameStatsTimerGraphSeries& Series = *NewSeries;
 
-	Series->SetName(TEXT("<Frame Stats Timer>"));
-	Series->SetDescription(TEXT("Frame Stats Timer series"));
+	Series.SetName(TEXT("<Frame Stats Timer>"));
+	Series.SetDescription(TEXT("Frame Stats Timer series"));
 
 	const FLinearColor BorderColor(Color.R + 0.4f, Color.G + 0.4f, Color.B + 0.4f, 1.0f);
-	Series->SetColor(Color, BorderColor);
-
-	Series->TimerId = TimerId;
-	Series->FrameType = FrameType;
+	Series.SetColor(Color, BorderColor);
 
 	// Use shared viewport.
-	Series->SetBaselineY(SharedValueViewport.GetBaselineY());
-	Series->SetScaleY(SharedValueViewport.GetScaleY());
-	Series->EnableSharedViewport();
+	Series.SetBaselineY(SharedValueViewport.GetBaselineY());
+	Series.SetScaleY(SharedValueViewport.GetScaleY());
+	Series.EnableSharedViewport();
 
-	Series->CachedSessionDuration = 0.0;
-
-	AllSeries.Add(Series);
-	return Series;
+	AllSeries.Add(NewSeries);
+	return NewSeries;
 }
 
 ////////////////////////////////////////////////////////////////////////////////////////////////////
 
 void FTimingGraphTrack::RemoveFrameStatsTimerSeries(uint32 TimerId, ETraceFrameType FrameType)
 {
-	AllSeries.RemoveAll([TimerId, FrameType](const TSharedPtr<FGraphSeries>& Series)
+	AllSeries.RemoveAll(
+		[TimerId, FrameType]
+		(const TSharedPtr<FGraphSeries>& Series)
 		{
-			const TSharedPtr<FTimingGraphSeries> TimingSeries = StaticCastSharedPtr<FTimingGraphSeries>(Series);
-			return TimingSeries->Type == FTimingGraphSeries::ESeriesType::FrameStatsTimer &&
-				   TimingSeries->TimerId == TimerId &&
-				   TimingSeries->FrameType == FrameType;
+			return Series->Is<FFrameStatsTimerGraphSeries>() &&
+				   Series->As<FFrameStatsTimerGraphSeries>().GetTimerId() == TimerId &&
+				   Series->As<FFrameStatsTimerGraphSeries>().GetFrameType() == FrameType;
 		});
 }
 
 ////////////////////////////////////////////////////////////////////////////////////////////////////
-
-void FTimingGraphTrack::UpdateFrameStatsTimerSeries(FTimingGraphSeries& Series, const FTimingTrackViewport& Viewport)
-{
-	using namespace UE::Insights::TimingProfiler;
-
-	FGraphTrackBuilder Builder(*this, Series, Viewport);
-	TSharedPtr<const TraceServices::IAnalysisSession> Session = UE::Insights::FInsightsManager::Get()->GetSession();
-	if (Session.IsValid())
-	{
-		TraceServices::FAnalysisSessionReadScope SessionReadScope(*Session.Get());
-
-		TSet<uint32> VisibleTimelines;
-		GetVisibleTimelineIndexes(VisibleTimelines);
-		const double SessionDuration = Session->GetDurationSeconds();
-		if (Series.CachedSessionDuration != SessionDuration || Series.CachedTimelinesNum != VisibleTimelines.Num())
-		{
-			Series.CachedSessionDuration = SessionDuration;
-			Series.CachedTimelinesNum = VisibleTimelines.Num();
-
-			const TraceServices::IFrameProvider& FramesProvider = ReadFrameProvider(*Session.Get());
-
-			Series.FrameStatsCachedEvents.Empty();
-			uint64 FrameCount = FramesProvider.GetFrameCount(ETraceFrameType::TraceFrameType_Game);
-			if (FrameCount == 0)
-			{
-				return;
-			}
-
-			FramesProvider.EnumerateFrames(Series.FrameType, 0ull, FrameCount, [&Series](const TraceServices::FFrame& Frame)
-				{
-					FFrameStatsCachedEvent Event;
-					Event.FrameStartTime = Frame.StartTime;
-					Event.FrameEndTime = Frame.EndTime;
-					Event.Duration.store(0.0f);
-					Series.FrameStatsCachedEvents.Add(Event);
-				});
-
-			FFrameStatsHelper::ComputeFrameStatsForTimer(Series.FrameStatsCachedEvents, Series.TimerId, VisibleTimelines);
-		}
-
-		int32 StartIndex = Algo::UpperBoundBy(Series.FrameStatsCachedEvents, Viewport.GetStartTime(), &FFrameStatsCachedEvent::FrameStartTime);
-		if (StartIndex > 0)
-		{
-			StartIndex--;
-		}
-		int32 EndIndex = Algo::UpperBoundBy(Series.FrameStatsCachedEvents, Viewport.GetEndTime(), &FFrameStatsCachedEvent::FrameStartTime);
-		if (EndIndex < Series.FrameStatsCachedEvents.Num())
-		{
-			EndIndex++;
-		}
-		for (int32 Index = StartIndex; Index < EndIndex; ++Index)
-		{
-			const FFrameStatsCachedEvent& Entry = Series.FrameStatsCachedEvents[Index];
-			Builder.AddEvent(Entry.FrameStartTime, Entry.Duration.load(), Entry.Duration.load());
-		}
-	}
-}
-
-////////////////////////////////////////////////////////////////////////////////////////////////////
-// Stats Counter Series
+// Counter Series
 ////////////////////////////////////////////////////////////////////////////////////////////////////
 
 TSharedPtr<FTimingGraphSeries> FTimingGraphTrack::GetStatsCounterSeries(uint32 CounterId)
 {
-	TSharedPtr<FGraphSeries>* Ptr = AllSeries.FindByPredicate([CounterId](const TSharedPtr<FGraphSeries>& Series)
-	{
-		const TSharedPtr<FTimingGraphSeries> TimingSeries = StaticCastSharedPtr<FTimingGraphSeries>(Series);
-		return TimingSeries->Type == FTimingGraphSeries::ESeriesType::StatsCounter && TimingSeries->CounterId == CounterId;
-	});
+	TSharedPtr<FGraphSeries>* Ptr = AllSeries.FindByPredicate(
+		[CounterId]
+		(const TSharedPtr<FGraphSeries>& Series)
+		{
+			return Series->Is<FCounterGraphSeries>() &&
+				   Series->As<FCounterGraphSeries>().GetCounterId() == CounterId;
+		});
 	return (Ptr != nullptr) ? StaticCastSharedPtr<FTimingGraphSeries>(*Ptr) : nullptr;
 }
 
@@ -697,140 +381,43 @@ TSharedPtr<FTimingGraphSeries> FTimingGraphTrack::GetStatsCounterSeries(uint32 C
 
 TSharedPtr<FTimingGraphSeries> FTimingGraphTrack::AddStatsCounterSeries(uint32 CounterId, FLinearColor Color)
 {
-	TSharedRef<FTimingGraphSeries> Series = MakeShared<FTimingGraphSeries>(FTimingGraphSeries::ESeriesType::StatsCounter);
-
-	const TCHAR* CounterName = nullptr;
-	bool bIsTime = false;
-	bool bIsMemory = false;
-	bool bIsFloatingPoint = false;
+	TSharedRef<FCounterGraphSeries> NewSeries = MakeShared<FCounterGraphSeries>(CounterId);
+	FCounterGraphSeries& Series = *NewSeries;
 
-	TSharedPtr<const TraceServices::IAnalysisSession> Session = UE::Insights::FInsightsManager::Get()->GetSession();
-	if (Session.IsValid())
-	{
-		TraceServices::FAnalysisSessionReadScope SessionReadScope(*Session.Get());
-		const TraceServices::ICounterProvider& CountersProvider = TraceServices::ReadCounterProvider(*Session.Get());
-		if (CounterId < CountersProvider.GetCounterCount())
-		{
-			CountersProvider.ReadCounter(CounterId, [&](const TraceServices::ICounter& Counter)
-			{
-				CounterName = Counter.GetName();
-				//bIsTime = (Counter.GetDisplayHint() == TraceServices::CounterDisplayHint_Time);
-				bIsMemory = (Counter.GetDisplayHint() == TraceServices::CounterDisplayHint_Memory);
-				bIsFloatingPoint = Counter.IsFloatingPoint();
-			});
-		}
-	}
-
-	Series->SetName(CounterName != nullptr ? CounterName : TEXT("<StatsCounter>"));
-	Series->SetDescription(TEXT("Stats counter series"));
+	Series.SetName(TEXT("<StatsCounter>"));
+	Series.SetDescription(TEXT("Stats counter series"));
 
 	FLinearColor BorderColor(Color.R + 0.4f, Color.G + 0.4f, Color.B + 0.4f, 1.0f);
-	Series->SetColor(Color, BorderColor);
-
-	Series->CounterId = CounterId;
+	Series.SetColor(Color, BorderColor);
 
-	Series->bIsTime = bIsTime;
-	Series->bIsMemory = bIsMemory;
-	Series->bIsFloatingPoint = bIsFloatingPoint;
+	Series.SetBaselineY(GetHeight() - 1.0f);
+	Series.SetScaleY(1.0);
+	Series.EnableAutoZoom();
 
-	Series->SetBaselineY(GetHeight() - 1.0f);
-	Series->SetScaleY(1.0);
-	Series->EnableAutoZoom();
+	Series.InitFromProvider();
 
-	AllSeries.Add(Series);
-	return Series;
+	AllSeries.Add(NewSeries);
+	return NewSeries;
 }
 
 ////////////////////////////////////////////////////////////////////////////////////////////////////
 
 void FTimingGraphTrack::RemoveStatsCounterSeries(uint32 CounterId)
 {
-	AllSeries.RemoveAll([CounterId](const TSharedPtr<FGraphSeries>& Series)
-	{
-		const TSharedPtr<FTimingGraphSeries> TimingSeries = StaticCastSharedPtr<FTimingGraphSeries>(Series);
-		return TimingSeries->Type == FTimingGraphSeries::ESeriesType::StatsCounter && TimingSeries->CounterId == CounterId;
-	});
-}
-
-////////////////////////////////////////////////////////////////////////////////////////////////////
-
-void FTimingGraphTrack::UpdateStatsCounterSeries(FTimingGraphSeries& Series, const FTimingTrackViewport& Viewport)
-{
-	FGraphTrackBuilder Builder(*this, Series, Viewport);
-
-	TSharedPtr<const TraceServices::IAnalysisSession> Session = UE::Insights::FInsightsManager::Get()->GetSession();
-	if (Session.IsValid())
-	{
-		TraceServices::FAnalysisSessionReadScope SessionReadScope(*Session.Get());
-		const TraceServices::ICounterProvider& CounterProvider = TraceServices::ReadCounterProvider(*Session.Get());
-		CounterProvider.ReadCounter(Series.CounterId, [this, &Viewport, &Builder, &Series](const TraceServices::ICounter& Counter)
+	AllSeries.RemoveAll(
+		[CounterId]
+		(const TSharedPtr<FGraphSeries>& Series)
 		{
-			const float TopY = 4.0f;
-			const float BottomY = GetHeight() - 4.0f;
-
-			if (Series.IsAutoZoomEnabled() && TopY < BottomY)
-			{
-				double MinValue =  std::numeric_limits<double>::infinity();
-				double MaxValue = -std::numeric_limits<double>::infinity();
-
-				if (Counter.IsFloatingPoint())
-				{
-					Counter.EnumerateFloatValues(Viewport.GetStartTime(), Viewport.GetEndTime(), true, [&Builder, &MinValue, &MaxValue](double Time, double Value)
-					{
-						if (Value < MinValue)
-						{
-							MinValue = Value;
-						}
-						if (Value > MaxValue)
-						{
-							MaxValue = Value;
-						}
-					});
-				}
-				else
-				{
-					Counter.EnumerateValues(Viewport.GetStartTime(), Viewport.GetEndTime(), true, [&Builder, &MinValue, &MaxValue](double Time, int64 IntValue)
-					{
-						const double Value = static_cast<double>(IntValue);
-						if (Value < MinValue)
-						{
-							MinValue = Value;
-						}
-						if (Value > MaxValue)
-						{
-							MaxValue = Value;
-						}
-					});
-				}
-
-				Series.UpdateAutoZoom(TopY, BottomY, MinValue, MaxValue);
-			}
-
-			if (Counter.IsFloatingPoint())
-			{
-				Counter.EnumerateFloatValues(Viewport.GetStartTime(), Viewport.GetEndTime(), true, [&Builder](double Time, double Value)
-				{
-					//TODO: add a "value unit converter"
-					Builder.AddEvent(Time, 0.0, Value);
-				});
-			}
-			else
-			{
-				Counter.EnumerateValues(Viewport.GetStartTime(), Viewport.GetEndTime(), true, [&Builder](double Time, int64 IntValue)
-				{
-					//TODO: add a "value unit converter"
-					Builder.AddEvent(Time, 0.0, static_cast<double>(IntValue));
-				});
-			}
+			return Series->Is<FCounterGraphSeries>() &&
+				   Series->As<FCounterGraphSeries>().GetCounterId() == CounterId;
 		});
-	}
 }
 
 ////////////////////////////////////////////////////////////////////////////////////////////////////
 
-TSharedPtr<UE::Insights::TimingProfiler::FThreadTimingSharedState> FTimingGraphTrack::GetThreadTimingSharedState() const
+TSharedPtr<FThreadTimingSharedState> FTimingGraphTrack::GetThreadTimingSharedState() const
 {
-	TSharedPtr<UE::Insights::TimingProfiler::STimingView> TimingViewPtr = TimingView.Pin();
+	TSharedPtr<STimingView> TimingViewPtr = TimingView.Pin();
 	if (!TimingViewPtr.IsValid())
 	{
 		return nullptr;
@@ -843,7 +430,7 @@ TSharedPtr<UE::Insights::TimingProfiler::FThreadTimingSharedState> FTimingGraphT
 
 void FTimingGraphTrack::GetVisibleTimelineIndexes(TSet<uint32>& TimelineIndexes)
 {
-	TSharedPtr<UE::Insights::TimingProfiler::FThreadTimingSharedState> ThreadSharedState = GetThreadTimingSharedState();
+	TSharedPtr<FThreadTimingSharedState> ThreadSharedState = GetThreadTimingSharedState();
 	if (ThreadSharedState)
 	{
 		ThreadSharedState->GetVisibleTimelineIndexes(TimelineIndexes);
@@ -854,7 +441,7 @@ void FTimingGraphTrack::GetVisibleTimelineIndexes(TSet<uint32>& TimelineIndexes)
 
 void FTimingGraphTrack::GetVisibleCpuSamplingThreads(TSet<uint32>& Threads)
 {
-	TSharedPtr<UE::Insights::TimingProfiler::FThreadTimingSharedState> ThreadSharedState = GetThreadTimingSharedState();
+	TSharedPtr<FThreadTimingSharedState> ThreadSharedState = GetThreadTimingSharedState();
 	if (ThreadSharedState)
 	{
 		ThreadSharedState->GetVisibleCpuSamplingThreads(Threads);
@@ -867,7 +454,7 @@ void FTimingGraphTrack::ContextMenu_ToggleOption_Execute(EGraphOptions Option)
 {
 	FGraphTrack::ContextMenu_ToggleOption_Execute(Option);
 
-	TSharedPtr<UE::Insights::TimingProfiler::STimingView> TimingViewPtr = TimingView.Pin();
+	TSharedPtr<STimingView> TimingViewPtr = TimingView.Pin();
 	if (!TimingViewPtr.IsValid())
 	{
 		return;
@@ -908,7 +495,7 @@ void FTimingGraphTrack::ContextMenu_ToggleOption_Execute(EGraphOptions Option)
 
 void FTimingGraphTrack::LoadDefaultSettings()
 {
-	TSharedPtr<UE::Insights::TimingProfiler::STimingView> TimingViewPtr = TimingView.Pin();
+	TSharedPtr<STimingView> TimingViewPtr = TimingView.Pin();
 	if (TimingViewPtr.IsValid() && TimingViewPtr->GetName() == FInsightsManagerTabs::TimingProfilerTabId)
 	{
 		const FInsightsSettings& Settings = UE::Insights::FInsightsManager::Get()->GetSettings();
@@ -943,17 +530,15 @@ void FTimingGraphTrack::LoadDefaultSettings()
 
 void FTimingGraphTrack::DrawVerticalAxisGrid(const ITimingTrackDrawContext& Context) const
 {
-	TSharedPtr<FTimingGraphSeries> FirstTimeUnitSeries;
+	FTimingGraphSeries* FirstTimeUnitSeries = nullptr;
 	for (const TSharedPtr<FGraphSeries>& Series : AllSeries)
 	{
-		if (Series->IsVisible())
+		if (Series->IsVisible() &&
+			Series->Is<FTimingGraphSeries>() &&
+			Series->As<FTimingGraphSeries>().IsTimeUnit())
 		{
-			TSharedPtr<FTimingGraphSeries> TimingSeries = StaticCastSharedPtr<FTimingGraphSeries>(Series);
-			if (TimingSeries->bIsTime)
-			{
-				FirstTimeUnitSeries = TimingSeries;
-				break;
-			}
+			FirstTimeUnitSeries = &Series->As<FTimingGraphSeries>();
+			break;
 		}
 	}
 	if (!FirstTimeUnitSeries)
@@ -1066,16 +651,12 @@ bool FTimingGraphTrack::HasAnySeriesForTimer(uint32 TimerId) const
 {
 	for (const TSharedPtr<FGraphSeries>& Series : AllSeries)
 	{
-		const TSharedPtr<FTimingGraphSeries> TimingSeries = StaticCastSharedPtr<FTimingGraphSeries>(Series);
-		if (TimingSeries.IsValid() &&
-			(TimingSeries->Type == FTimingGraphSeries::ESeriesType::Timer ||
-			 TimingSeries->Type == FTimingGraphSeries::ESeriesType::FrameStatsTimer) &&
-			TimingSeries->TimerId == TimerId)
+		if (Series->Is<FTimingGraphSeries>() &&
+			Series->As<FTimingGraphSeries>().IsTimer(TimerId))
 		{
 			return true;
 		}
 	}
-
 	return false;
 }
 
@@ -1084,22 +665,19 @@ bool FTimingGraphTrack::HasAnySeriesForTimer(uint32 TimerId) const
 uint32 FTimingGraphTrack::GetNumSeriesForTimer(uint32 TimerId) const
 {
 	uint32 NumSeries = 0;
-
 	for (const TSharedPtr<FGraphSeries>& Series : AllSeries)
 	{
-		const TSharedPtr<FTimingGraphSeries> TimingSeries = StaticCastSharedPtr<FTimingGraphSeries>(Series);
-		if (TimingSeries.IsValid() &&
-			(TimingSeries->Type == FTimingGraphSeries::ESeriesType::Timer ||
-			 TimingSeries->Type == FTimingGraphSeries::ESeriesType::FrameStatsTimer) &&
-			TimingSeries->TimerId == TimerId)
+		if (Series->Is<FTimingGraphSeries>() &&
+			Series->As<FTimingGraphSeries>().IsTimer(TimerId))
 		{
 			++NumSeries;
 		}
 	}
-
 	return NumSeries;
 }
 
 ////////////////////////////////////////////////////////////////////////////////////////////////////
 
+} // namespace UE::Insights::TimingProfiler
+
 #undef LOCTEXT_NAMESPACE
diff --git a/Engine/Source/Developer/TraceInsights/Private/Insights/TimingProfiler/GraphTracks/TimingGraphTrack.h b/Engine/Source/Developer/TraceInsights/Private/Insights/TimingProfiler/GraphTracks/TimingGraphTrack.h
--- a/Engine/Source/Developer/TraceInsights/Private/Insights/TimingProfiler/GraphTracks/TimingGraphTrack.h
+++ b/Engine/Source/Developer/TraceInsights/Private/Insights/TimingProfiler/GraphTracks/TimingGraphTrack.h
@@ -11,82 +11,23 @@
 
 namespace UE::Insights::TimingProfiler
 {
-	struct FFrameStatsCachedEvent;
-	class STimingView;
-}
 
-////////////////////////////////////////////////////////////////////////////////////////////////////
+struct FFrameStatsCachedEvent;
+class STimingView;
 
-/** The delegate to be invoked when a series visibility is changed. */
-DECLARE_MULTICAST_DELEGATE_OneParam(FSeriesVisibilityChangedDelegate, bool bOnOff);
-
-class FTimingGraphSeries : public FGraphSeries
-{
-	INSIGHTS_DECLARE_RTTI(FTimingGraphSeries, FGraphSeries)
-
-public:
-	enum class ESeriesType
-	{
-		Frame,
-		Timer,
-		StatsCounter,
-		FrameStatsTimer
-	};
-
-	struct FSimpleTimingEvent
-	{
-		double StartTime;
-		double Duration;
-	};
-
-public:
-	explicit FTimingGraphSeries(FTimingGraphSeries::ESeriesType Type);
-	virtual ~FTimingGraphSeries();
-
-	virtual FString FormatValue(double Value) const override;
-
-	static bool CompareEventsByStartTime(const FSimpleTimingEvent& EventA, const FSimpleTimingEvent& EventB)
-	{
-		return EventA.StartTime < EventB.StartTime;
-	}
-
-	virtual void SetVisibility(bool bOnOff) override;
-
-public:
-	ESeriesType Type;
-	union
-	{
-		uint32 TimerId;
-		uint32 CounterId;
-	};
-
-	ETraceFrameType FrameType;
-	double CachedSessionDuration;
-	TArray<FSimpleTimingEvent> CachedEvents; // used by Timer series
-	TArray<UE::Insights::TimingProfiler::FFrameStatsCachedEvent> FrameStatsCachedEvents; // used by Frame Stats Timer series
-	uint32 CachedTimelinesNum = 0; // the number of timelines used to gather the data
-
-	bool bIsTime; // the unit for values is [second]
-	bool bIsMemory; // the unit for value is [byte]
-	bool bIsFloatingPoint; // for stats counters
-
-	FSeriesVisibilityChangedDelegate VisibilityChangedDelegate;
-};
-
-////////////////////////////////////////////////////////////////////////////////////////////////////
+class FTimingGraphSeries;
 
 class FTimingGraphTrack : public FGraphTrack
 {
 	INSIGHTS_DECLARE_RTTI(FTimingGraphTrack, FGraphTrack)
 
 public:
-	FTimingGraphTrack(TSharedPtr<UE::Insights::TimingProfiler::STimingView> InTimingView);
+	FTimingGraphTrack(TSharedPtr<STimingView> InTimingView);
 	virtual ~FTimingGraphTrack();
 
 	virtual void Update(const ITimingTrackUpdateContext& Context) override;
 
 	void AddDefaultFrameSeries();
-
 	TSharedPtr<FTimingGraphSeries> GetFrameSeries(ETraceFrameType FrameType);
 
 	TSharedPtr<FTimingGraphSeries> GetTimerSeries(uint32 TimerId);
@@ -104,16 +45,15 @@ public:
 	bool HasAnySeriesForTimer(uint32 TimerId) const;
 	uint32 GetNumSeriesForTimer(uint32 TimerId) const;
 
-protected:
-	void UpdateFrameSeries(FTimingGraphSeries& Series, const FTimingTrackViewport& Viewport);
-	void UpdateTimerSeries(FTimingGraphSeries& Series, const FTimingTrackViewport& Viewport);
-	void UpdateFrameStatsTimerSeries(FTimingGraphSeries& Series, const FTimingTrackViewport& Viewport);
-	void UpdateStatsCounterSeries(FTimingGraphSeries& Series, const FTimingTrackViewport& Viewport);
-
-	TSharedPtr<UE::Insights::TimingProfiler::FThreadTimingSharedState> GetThreadTimingSharedState() const;
 	void GetVisibleTimelineIndexes(TSet<uint32>& TimelineIndexes);
 	void GetVisibleCpuSamplingThreads(TSet<uint32>& Threads);
 
+private:
+	void RegisterTimingViewCallbacks();
+	void UnregisterTimingViewCallbacks();
+
+	TSharedPtr<FThreadTimingSharedState> GetThreadTimingSharedState() const;
+
 	virtual void DrawVerticalAxisGrid(const ITimingTrackDrawContext& Context) const override;
 
 	void LoadDefaultSettings();
@@ -129,8 +69,8 @@ private:
 	FDelegateHandle GameFrameSeriesVisibilityHandle;
 	FDelegateHandle RenderingFrameSeriesVisibilityHandle;
 
-	TWeakPtr<UE::Insights::TimingProfiler::STimingView> TimingView;
+	TWeakPtr<STimingView> TimingView;
 	bool bNotifyTimersOnDestruction;
 };
 
-////////////////////////////////////////////////////////////////////////////////////////////////////
+} // namespace UE::Insights::TimingProfiler
diff --git a/Engine/Source/Developer/TraceInsights/Private/Insights/TimingProfiler/ViewModels/FrameStatsHelper.cpp b/Engine/Source/Developer/TraceInsights/Private/Insights/TimingProfiler/ViewModels/FrameStatsHelper.cpp
--- a/Engine/Source/Developer/TraceInsights/Private/Insights/TimingProfiler/ViewModels/FrameStatsHelper.cpp
+++ b/Engine/Source/Developer/TraceInsights/Private/Insights/TimingProfiler/ViewModels/FrameStatsHelper.cpp
@@ -15,6 +15,11 @@ namespace UE::Insights::TimingProfiler
 
 void FFrameStatsHelper::ComputeFrameStatsForTimer(TArray<FFrameStatsCachedEvent>& FrameStatsEvents, uint32 TimerId, const TSet<uint32>& Timelines)
 {
+	if (FrameStatsEvents.Num() == 0)
+	{
+		return;
+	}
+
 	TSharedPtr<const TraceServices::IAnalysisSession> Session = FInsightsManager::Get()->GetSession();
 	if (Session.IsValid())
 	{
@@ -26,7 +31,17 @@ void FFrameStatsHelper::ComputeFrameStatsForTimer(TArray<FFrameStatsCachedEvent>
 
 			for (uint32 TimelineIndex : Timelines)
 			{
-				ProcessTimeline(*TimingProfilerProvider, SessionDuration, FrameStatsEvents, TimerId, TimelineIndex);
+				TimingProfilerProvider->ReadTimeline(TimelineIndex,
+					[&FrameStatsEvents, TimerId, SessionDuration, TimingProfilerProvider]
+					(const TraceServices::ITimingProfilerTimeline& Timeline)
+					{
+						TimingProfilerProvider->ReadTimers(
+							[&FrameStatsEvents, TimerId, SessionDuration, &Timeline]
+							(const TraceServices::ITimingProfilerTimerReader& TimerReader)
+							{
+								ProcessTimeline(FrameStatsEvents, TimerId, 0.0, SessionDuration, Timeline, TimerReader);
+							});
+					});
 			}
 		}
 	}
@@ -36,6 +51,11 @@ void FFrameStatsHelper::ComputeFrameStatsForTimer(TArray<FFrameStatsCachedEvent>
 
 void FFrameStatsHelper::ComputeFrameStatsForTimer(TArray<FFrameStatsCachedEvent>& FrameStatsEvents, uint32 TimerId)
 {
+	if (FrameStatsEvents.Num() == 0)
+	{
+		return;
+	}
+
 	TSharedPtr<const TraceServices::IAnalysisSession> Session = FInsightsManager::Get()->GetSession();
 	if (Session.IsValid())
 	{
@@ -45,112 +65,149 @@ void FFrameStatsHelper::ComputeFrameStatsForTimer(TArray<FFrameStatsCachedEvent>
 			Session->ReadAccessCheck();
 			const double SessionDuration = Session->GetDurationSeconds();
 
-			const uint32 TimelineCount = TimingProfilerProvider->GetTimelineCount();
-			for (uint32 TimelineIndex = 0; TimelineIndex < TimelineCount; ++TimelineIndex)
-			{
-				ProcessTimeline(*TimingProfilerProvider, SessionDuration, FrameStatsEvents, TimerId, TimelineIndex);
-			}
+			TimingProfilerProvider->EnumerateTimelines(
+				[&FrameStatsEvents, TimerId, SessionDuration, TimingProfilerProvider]
+				(const TraceServices::ITimingProfilerTimeline& Timeline)
+				{
+					TimingProfilerProvider->ReadTimers(
+						[&FrameStatsEvents, TimerId, SessionDuration, &Timeline]
+						(const TraceServices::ITimingProfilerTimerReader& TimerReader)
+						{
+							ProcessTimeline(FrameStatsEvents, TimerId, 0.0, SessionDuration, Timeline, TimerReader);
+						});
+				});
 		}
 	}
 }
 
 ////////////////////////////////////////////////////////////////////////////////////////////////////
 
-void FFrameStatsHelper::ProcessTimeline(const TraceServices::ITimingProfilerProvider& TimingProfilerProvider, double SessionDuration, TArray<FFrameStatsCachedEvent>& FrameStatsEvents, uint32 TimerId, uint32 TimelineIndex)
+void FFrameStatsHelper::ComputeFrameStatsForTimer(TArray<FFrameStatsCachedEvent>& FrameStatsEvents, uint32 TimerId, const TraceServices::ITimingProfilerTimeline& Timeline)
 {
-	const TraceServices::ITimingProfilerTimerReader* TimerReader;
-	TimingProfilerProvider.ReadTimers([&TimerReader](const TraceServices::ITimingProfilerTimerReader& Out) { TimerReader = &Out; });
+	if (FrameStatsEvents.Num() == 0)
+	{
+		return;
+	}
 
-	TimingProfilerProvider.ReadTimeline(TimelineIndex,
-		[SessionDuration, &FrameStatsEvents, TimerReader, TimerId]
-		(const TraceServices::ITimingProfilerTimeline& Timeline)
+	TSharedPtr<const TraceServices::IAnalysisSession> Session = FInsightsManager::Get()->GetSession();
+	if (Session.IsValid())
+	{
+		const TraceServices::ITimingProfilerProvider* TimingProfilerProvider = TraceServices::ReadTimingProfilerProvider(*Session.Get());
+		if (TimingProfilerProvider)
 		{
-			struct FEnumerateAsyncTaskData
+			Session->ReadAccessCheck();
+			const double SessionDuration = Session->GetDurationSeconds();
+
+			TimingProfilerProvider->ReadTimers(
+				[&FrameStatsEvents, TimerId, &Timeline, SessionDuration]
+				(const TraceServices::ITimingProfilerTimerReader& TimerReader)
+				{
+					ProcessTimeline(FrameStatsEvents, TimerId, 0.0, SessionDuration, Timeline, TimerReader);
+				});
+		}
+	}
+}
+
+////////////////////////////////////////////////////////////////////////////////////////////////////
+
+void FFrameStatsHelper::ProcessTimeline(
+	TArray<FFrameStatsCachedEvent>& FrameStatsEvents,
+	uint32 TimerId,
+	double StartTime,
+	double EndTime,
+	const TraceServices::ITimingProfilerTimeline& Timeline,
+	const TraceServices::ITimingProfilerTimerReader& TimerReader)
+{
+	struct FEnumerateAsyncTaskData
+	{
+		double StartTime;
+		int32 NestedDepth = 0;
+	};
+
+	TArray<FEnumerateAsyncTaskData> DataArray;
+
+	TraceServices::ITimingProfilerTimeline::EnumerateAsyncParams Params;
+	Params.IntervalStart = StartTime;
+	Params.IntervalEnd = EndTime;
+	Params.Resolution = 0.0;
+	Params.SetupCallback =
+		[&DataArray]
+		(uint32 NumTasks)
+		{
+			DataArray.AddDefaulted(NumTasks);
+		};
+	Params.EventCallback =
+		[&TimerReader, &FrameStatsEvents, TimerId, &DataArray]
+		(bool bIsEnter, double Time, const TraceServices::FTimingProfilerEvent& Event, uint32 TaskIndex)
+		{
+			const TraceServices::FTimingProfilerTimer* Timer = TimerReader.GetTimer(Event.TimerIndex);
+			if (ensure(Timer != nullptr))
 			{
-				double StartTime;
-				int32 NestedDepth = 0;
-			};
-
-			TArray<FEnumerateAsyncTaskData> DataArray;
-
-			TraceServices::ITimingProfilerTimeline::EnumerateAsyncParams Params;
-			Params.IntervalStart = 0;
-			Params.IntervalEnd = SessionDuration;
-			Params.Resolution = 0.0;
-			Params.SetupCallback = [&DataArray](uint32 NumTasks) { DataArray.AddDefaulted(NumTasks); };
-			Params.EventCallback =
-				[TimerReader, &FrameStatsEvents, TimerId, &DataArray]
-				(bool bIsEnter, double Time, const TraceServices::FTimingProfilerEvent& Event, uint32 TaskIndex)
+				if (Timer->Id == TimerId)
 				{
-					const TraceServices::FTimingProfilerTimer* Timer = TimerReader->GetTimer(Event.TimerIndex);
-					if (ensure(Timer != nullptr))
+					FEnumerateAsyncTaskData& CurrentTaskData = DataArray[TaskIndex];
+					if (bIsEnter)
 					{
-						if (Timer->Id == TimerId)
+						if (CurrentTaskData.NestedDepth == 0)
+						{
+							CurrentTaskData.StartTime = Time;
+						}
+
+						++CurrentTaskData.NestedDepth;
+					}
+					else
+					{
+						check(CurrentTaskData.NestedDepth > 0);
+						if (--CurrentTaskData.NestedDepth > 0)
+						{
+							return TraceServices::EEventEnumerate::Continue;
+						}
+
+						int32 Index = Algo::UpperBoundBy(FrameStatsEvents, CurrentTaskData.StartTime, &FFrameStatsCachedEvent::FrameStartTime);
+						if (Index > 0)
+						{
+							--Index;
+						}
+
+						// This can happen when the event is between frames.
+						if (CurrentTaskData.StartTime > FrameStatsEvents[Index].FrameEndTime)
 						{
-							FEnumerateAsyncTaskData& CurrentTaskData = DataArray[TaskIndex];
-							if (bIsEnter)
+							Index++;
+							if (Index >= FrameStatsEvents.Num())
 							{
-								if (CurrentTaskData.NestedDepth == 0)
-								{
-									CurrentTaskData.StartTime = Time;
-								}
+								return TraceServices::EEventEnumerate::Continue;
+							}
+						}
 
-								++CurrentTaskData.NestedDepth;
+						double EndTime = Time;
+						do
+						{
+							FFrameStatsCachedEvent& Entry = FrameStatsEvents[Index];
+
+							if (EndTime < Entry.FrameStartTime)
+							{
+								return TraceServices::EEventEnumerate::Continue;
 							}
-							else
+
+							if (CurrentTaskData.StartTime < Entry.FrameStartTime)
 							{
-								check(CurrentTaskData.NestedDepth > 0);
-								if (--CurrentTaskData.NestedDepth > 0)
-								{
-									return TraceServices::EEventEnumerate::Continue;
-								}
-
-								int32 Index = Algo::UpperBoundBy(FrameStatsEvents, CurrentTaskData.StartTime, &FFrameStatsCachedEvent::FrameStartTime);
-								if (Index > 0)
-								{
-									--Index;
-								}
-
-								// This can happen when the event is between frames.
-								if (CurrentTaskData.StartTime > FrameStatsEvents[Index].FrameEndTime)
-								{
-									Index++;
-									if (Index >= FrameStatsEvents.Num())
-									{
-										return TraceServices::EEventEnumerate::Continue;
-									}
-								}
-
-								double EndTime = Time;
-								do
-								{
-									FFrameStatsCachedEvent& Entry = FrameStatsEvents[Index];
-
-									if (EndTime < Entry.FrameStartTime)
-									{
-										return TraceServices::EEventEnumerate::Continue;
-									}
-
-									if (CurrentTaskData.StartTime < Entry.FrameStartTime)
-									{
-										CurrentTaskData.StartTime = Entry.FrameStartTime;
-									}
-
-									const double Duration = FMath::Min(EndTime, Entry.FrameEndTime) - CurrentTaskData.StartTime;
-									ensure(Duration >= 0.0f);
-									for (double Value = Entry.Duration.load(); !Entry.Duration.compare_exchange_strong(Value, Value + Duration););
-
-									Index++;
-								}
-								while (Index < FrameStatsEvents.Num());
+								CurrentTaskData.StartTime = Entry.FrameStartTime;
 							}
+
+							const double Duration = FMath::Min(EndTime, Entry.FrameEndTime) - CurrentTaskData.StartTime;
+							ensure(Duration >= 0.0f);
+							for (double Value = Entry.Duration.load(); !Entry.Duration.compare_exchange_strong(Value, Value + Duration););
+
+							Index++;
 						}
+						while (Index < FrameStatsEvents.Num());
 					}
-					return TraceServices::EEventEnumerate::Continue;
-				};
+				}
+			}
+			return TraceServices::EEventEnumerate::Continue;
+		};
 
-			Timeline.EnumerateEventsDownSampledAsync(Params);
-		});
+	Timeline.EnumerateEventsDownSampledAsync(Params);
 }
 
 ////////////////////////////////////////////////////////////////////////////////////////////////////
diff --git a/Engine/Source/Developer/TraceInsights/Private/Insights/TimingProfiler/ViewModels/FrameStatsHelper.h b/Engine/Source/Developer/TraceInsights/Private/Insights/TimingProfiler/ViewModels/FrameStatsHelper.h
--- a/Engine/Source/Developer/TraceInsights/Private/Insights/TimingProfiler/ViewModels/FrameStatsHelper.h
+++ b/Engine/Source/Developer/TraceInsights/Private/Insights/TimingProfiler/ViewModels/FrameStatsHelper.h
@@ -6,12 +6,10 @@
 
 #include "Containers/Array.h"
 
-#include <atomic>
+// TraceServices
+#include "TraceServices/Model/TimingProfiler.h"
 
-namespace TraceServices
-{
-	class ITimingProfilerProvider;
-}
+#include <atomic>
 
 namespace UE::Insights::TimingProfiler
 {
@@ -19,10 +17,10 @@ namespace UE::Insights::TimingProfiler
 struct FFrameStatsCachedEvent
 {
 	FFrameStatsCachedEvent()
-		: FrameStartTime(0.0f)
-		, FrameEndTime(0.0f)
+		: FrameStartTime(0.0)
+		, FrameEndTime(0.0)
 	{
-		Duration.store(0.0f);
+		Duration.store(0.0);
 	}
 
 	FFrameStatsCachedEvent(const FFrameStatsCachedEvent& Other)
@@ -44,13 +42,15 @@ class FFrameStatsHelper
 public:
 	static void ComputeFrameStatsForTimer(TArray<FFrameStatsCachedEvent>& FrameStatsEvents, uint32 TimerId, const TSet<uint32>& Timelines);
 	static void ComputeFrameStatsForTimer(TArray<FFrameStatsCachedEvent>& FrameStatsEvents, uint32 TimerId);
+	static void ComputeFrameStatsForTimer(TArray<FFrameStatsCachedEvent>& FrameStatsEvents, uint32 TimerId, const TraceServices::ITimingProfilerTimeline& Timeline);
 
 private:
-	static void ProcessTimeline(const TraceServices::ITimingProfilerProvider& TimingProfilerProvider,
-								double SessionDuration,
-								TArray<FFrameStatsCachedEvent>& FrameStatsEvents,
+	static void ProcessTimeline(TArray<FFrameStatsCachedEvent>& FrameStatsEvents,
 								uint32 TimerId,
-								uint32 TimelineIndex);
+								double StartTime,
+								double EndTime,
+								const TraceServices::ITimingProfilerTimeline& Timeline,
+								const TraceServices::ITimingProfilerTimerReader& TimerReader);
 };
 
 } // namespace UE::Insights::TimingProfiler
diff --git a/Engine/Source/Developer/TraceInsights/Private/Insights/TimingProfiler/Widgets/SFrameTrack.cpp b/Engine/Source/Developer/TraceInsights/Private/Insights/TimingProfiler/Widgets/SFrameTrack.cpp
--- a/Engine/Source/Developer/TraceInsights/Private/Insights/TimingProfiler/Widgets/SFrameTrack.cpp
+++ b/Engine/Source/Developer/TraceInsights/Private/Insights/TimingProfiler/Widgets/SFrameTrack.cpp
@@ -15,7 +15,9 @@
 #include "Widgets/Layout/SScrollBar.h"
 
 // TraceServices
+#include "Common/ProviderLock.h"
 #include "TraceServices/Model/Frames.h"
+#include "TraceServices/Model/StackSamples.h"
 #include "TraceServices/Model/TimingProfiler.h"
 
 // TraceInsightsCore
@@ -57,32 +59,15 @@ SFrameTrack::SFrameTrack()
 
 SFrameTrack::~SFrameTrack()
 {
-	TSharedPtr<STimingProfilerWindow> ProfilerWindow = FTimingProfilerManager::Get()->GetProfilerWindow();
-
-	if (OnTrackVisibilityChangedHandle.IsValid())
-	{
-		if (ProfilerWindow.IsValid())
-		{
-			TSharedPtr<STimingView> TimingView = ProfilerWindow->GetTimingView();
-			if (TimingView.IsValid())
-			{
-				if (TimingView.Get() == RegisteredTimingView)
-				{
-					TimingView->OnTrackVisibilityChanged().Remove(OnTrackVisibilityChangedHandle);
-					TimingView->OnTrackAdded().Remove(OnTrackAddedHandle);
-					TimingView->OnTrackRemoved().Remove(OnTrackRemovedHandle);
-				}
-			}
-		}
-	}
+	UnregisterTimingViewCallbacks();
 
+	TSharedPtr<STimingProfilerWindow> ProfilerWindow = FTimingProfilerManager::Get()->GetProfilerWindow();
 	if (ProfilerWindow.IsValid())
 	{
 		TArray<uint32> TimerIds;
 		for (TSharedPtr<FFrameTrackSeries>& Series : AllSeries)
 		{
-			if (Series.IsValid() &&
-				Series->Is<FTimerFrameStatsTrackSeries>())
+			if (Series->Is<FTimerFrameStatsTrackSeries>())
 			{
 				TimerIds.Add(Series->As<FTimerFrameStatsTrackSeries>().GetTimerId());
 			}
@@ -192,43 +177,68 @@ void SFrameTrack::Construct(const FArguments& InArgs)
 
 ////////////////////////////////////////////////////////////////////////////////////////////////////
 
-void SFrameTrack::Tick(const FGeometry& AllottedGeometry, const double InCurrentTime, const float InDeltaTime)
+void SFrameTrack::RegisterTimingViewCallbacks()
 {
-	TSharedPtr<STimingProfilerWindow> TimingWindow = FTimingProfilerManager::Get()->GetProfilerWindow();
-	if (TimingWindow.IsValid())
+	TSharedPtr<STimingView> TimingView = TimingViewWeak.Pin();
+	if (TimingView.IsValid())
 	{
-		TSharedPtr<STimingView> TimingView = TimingWindow->GetTimingView();
-		if (TimingView.IsValid())
-		{
-			if (!OnTrackVisibilityChangedHandle.IsValid() || TimingView.Get() != RegisteredTimingView)
+		auto OnTrackAddedRemovedLambda =
+			[this](const TSharedPtr<const FBaseTimingTrack> Track)
 			{
-				RegisteredTimingView = TimingView.Get();
-				this->bIsStateDirty = true;
-
-				auto OnTrackAddedRemovedLamda = [this](const TSharedPtr<const FBaseTimingTrack> Track)
+				if (Track->Is<FThreadTimingTrack>())
 				{
-					if (Track->Is<FThreadTimingTrack>())
+					// If there are more series than the default frame series.
+					if (AllSeries.Num() > ETraceFrameType::TraceFrameType_Count)
 					{
-						// If there are more series than the default frame series.
-						if (this->AllSeries.Num() > ETraceFrameType::TraceFrameType_Count)
-						{
-							this->bIsStateDirty = true;
-						}
+						bIsStateDirty = true;
 					}
-				};
+				}
+			};
 
-				OnTrackAddedHandle = TimingView->OnTrackAdded().AddLambda(OnTrackAddedRemovedLamda);
-				OnTrackRemovedHandle = TimingView->OnTrackRemoved().AddLambda(OnTrackAddedRemovedLamda);
+		OnTrackAddedHandle = TimingView->OnTrackAdded().AddLambda(OnTrackAddedRemovedLambda);
+		OnTrackRemovedHandle = TimingView->OnTrackRemoved().AddLambda(OnTrackAddedRemovedLambda);
 
-				OnTrackVisibilityChangedHandle = TimingView->OnTrackVisibilityChanged().AddLambda([this]()
-					{
-						if (this->AllSeries.Num() > ETraceFrameType::TraceFrameType_Count)
-						{
-							this->bIsStateDirty = true;
-						}
-					});
-			}
-		}
+		OnTrackVisibilityChangedHandle = TimingView->OnTrackVisibilityChanged().AddLambda(
+			[this]()
+			{
+				// If there are more series than the default frame series.
+				if (AllSeries.Num() > ETraceFrameType::TraceFrameType_Count)
+				{
+					bIsStateDirty = true;
+				}
+			});
+	}
+}
+
+////////////////////////////////////////////////////////////////////////////////////////////////////
+
+void SFrameTrack::UnregisterTimingViewCallbacks()
+{
+	TSharedPtr<STimingView> TimingView = TimingViewWeak.Pin();
+	if (TimingView.IsValid())
+	{
+		TimingView->OnTrackAdded().Remove(OnTrackAddedHandle);
+		TimingView->OnTrackRemoved().Remove(OnTrackRemovedHandle);
+		TimingView->OnTrackVisibilityChanged().Remove(OnTrackVisibilityChangedHandle);
+	}
+
+	OnTrackAddedHandle.Reset();
+	OnTrackRemovedHandle.Reset();
+	OnTrackVisibilityChangedHandle.Reset();
+}
+
+////////////////////////////////////////////////////////////////////////////////////////////////////
+
+void SFrameTrack::Tick(const FGeometry& AllottedGeometry, const double InCurrentTime, const float InDeltaTime)
+{
+	TSharedPtr<STimingView> OldTimingView = TimingViewWeak.Pin();
+	TSharedPtr<STimingView> NewTimingView = GetTimingView();
+	if (OldTimingView.Get() != NewTimingView.Get())
+	{
+		UnregisterTimingViewCallbacks();
+		TimingViewWeak = NewTimingView;
+		RegisterTimingViewCallbacks();
+		bIsStateDirty = true;
 	}
 
 	if (ThisGeometry != AllottedGeometry || bIsViewportDirty)
@@ -312,10 +322,12 @@ void SFrameTrack::Tick(const FGeometry& AllottedGeometry, const double InCurrent
 
 TSharedRef<FFrameTrackSeries> SFrameTrack::FindOrAddSeries(ETraceFrameType FrameType)
 {
-	TSharedPtr<FFrameTrackSeries>* ExistingSeries = AllSeries.FindByPredicate([FrameType](TSharedPtr<FFrameTrackSeries> Series)
+	TSharedPtr<FFrameTrackSeries>* ExistingSeries = AllSeries.FindByPredicate(
+		[FrameType]
+		(TSharedPtr<FFrameTrackSeries> Series)
 		{
 			return !Series->Is<FTimerFrameStatsTrackSeries>() &&
-				Series->GetFrameType() == FrameType;
+				   Series->GetFrameType() == FrameType;
 		});
 
 	if (ExistingSeries != nullptr)
@@ -336,10 +348,12 @@ TSharedRef<FFrameTrackSeries> SFrameTrack::FindOrAddSeries(ETraceFrameType Frame
 
 TSharedPtr<FFrameTrackSeries> SFrameTrack::FindSeries(ETraceFrameType FrameType) const
 {
-	const TSharedPtr<FFrameTrackSeries>* ExistingSeries = AllSeries.FindByPredicate([FrameType](TSharedPtr<FFrameTrackSeries> Series)
+	const TSharedPtr<FFrameTrackSeries>* ExistingSeries = AllSeries.FindByPredicate(
+		[FrameType]
+		(TSharedPtr<FFrameTrackSeries> Series)
 		{
 			return !Series->Is<FTimerFrameStatsTrackSeries>() &&
-				Series->GetFrameType() == FrameType;
+				   Series->GetFrameType() == FrameType;
 		});
 
 	if (ExistingSeries != nullptr)
@@ -354,11 +368,13 @@ TSharedPtr<FFrameTrackSeries> SFrameTrack::FindSeries(ETraceFrameType FrameType)
 
 TSharedPtr<FFrameTrackSeries> SFrameTrack::FindFrameStatsSeries(ETraceFrameType FrameType, uint32 TimerId) const
 {
-	const TSharedPtr<FFrameTrackSeries>* ExistingSeries = AllSeries.FindByPredicate([FrameType, TimerId](TSharedPtr<FFrameTrackSeries> Series)
+	const TSharedPtr<FFrameTrackSeries>* ExistingSeries = AllSeries.FindByPredicate(
+		[FrameType, TimerId]
+		(TSharedPtr<FFrameTrackSeries> Series)
 		{
 			return Series->Is<FTimerFrameStatsTrackSeries>() &&
-					Series->As<FTimerFrameStatsTrackSeries>().GetFrameType() == FrameType &&
-					Series->As<FTimerFrameStatsTrackSeries>().GetTimerId() == TimerId;
+				   Series->As<FTimerFrameStatsTrackSeries>().GetFrameType() == FrameType &&
+				   Series->As<FTimerFrameStatsTrackSeries>().GetTimerId() == TimerId;
 		});
 
 	if (ExistingSeries)
@@ -373,6 +389,8 @@ TSharedPtr<FFrameTrackSeries> SFrameTrack::FindFrameStatsSeries(ETraceFrameType
 
 void SFrameTrack::UpdateState()
 {
+	LLM_SCOPE_BYTAG(Insights);
+
 	FStopwatch Stopwatch;
 	Stopwatch.Start();
 
@@ -395,21 +413,24 @@ void SFrameTrack::UpdateState()
 		const uint64 StartIndex = static_cast<uint64>(FMath::Max(0, ViewportX.GetValueAtOffset(0.0f)));
 		const uint64 EndIndex = static_cast<uint64>(ViewportX.GetValueAtOffset(ViewportX.GetSize()));
 
+		// Frame Series
 		for (int32 FrameType = 0; FrameType < TraceFrameType_Count; ++FrameType)
 		{
 			TSharedPtr<FFrameTrackSeries> SeriesPtr = FindOrAddSeries(static_cast<ETraceFrameType>(FrameType));
 
-			LLM_SCOPE_BYTAG(Insights);
 			FFrameTrackSeriesBuilder Builder(*SeriesPtr, Viewport);
 
-			FramesProvider.EnumerateFrames(static_cast<ETraceFrameType>(FrameType), StartIndex, EndIndex, [&Builder](const TraceServices::FFrame& Frame)
-			{
-				Builder.AddFrame(Frame);
-			});
+			FramesProvider.EnumerateFrames(static_cast<ETraceFrameType>(FrameType), StartIndex, EndIndex,
+				[&Builder]
+				(const TraceServices::FFrame& Frame)
+				{
+					Builder.AddFrame(Frame);
+				});
 
 			NumUpdatedFrames += Builder.GetNumAddedFrames();
 		}
 
+		// Timer Series
 		for (int32 Index = 0; Index < AllSeries.Num(); ++Index)
 		{
 			TSharedPtr<FFrameTrackSeries> Series = AllSeries[Index];
@@ -417,10 +438,15 @@ void SFrameTrack::UpdateState()
 			{
 				continue;
 			}
+			const FTimerFrameStatsTrackSeries& TimerSeries = Series->As<FTimerFrameStatsTrackSeries>();
+			const uint32 TimerId = TimerSeries.GetTimerId();
+
+			FFrameTrackSeriesBuilder Builder(*Series, Viewport);
 
-			FTimerFrameStatsTrackSeries& TimerSeries = Series->As<FTimerFrameStatsTrackSeries>();
 			TArray<FFrameStatsCachedEvent> Frames;
-			FramesProvider.EnumerateFrames(static_cast<ETraceFrameType>(Series->GetFrameType()), StartIndex, EndIndex, [&Frames](const TraceServices::FFrame& Frame)
+			FramesProvider.EnumerateFrames(static_cast<ETraceFrameType>(Series->GetFrameType()), StartIndex, EndIndex,
+				[&Frames]
+				(const TraceServices::FFrame& Frame)
 				{
 					FFrameStatsCachedEvent Event;
 					Event.FrameStartTime = Frame.StartTime;
@@ -429,43 +455,57 @@ void SFrameTrack::UpdateState()
 					Frames.Add(Event);
 				});
 
-			FFrameTrackSeriesBuilder Builder(*Series, Viewport);
-
-			bool bTimingViewExists = false;
-			TSet<uint32> Timelines;
-
-			TSharedPtr<STimingProfilerWindow> TimingWindow = FTimingProfilerManager::Get()->GetProfilerWindow();
-
-			// Attempt to compute only from visible timelines.
-			if (TimingWindow.IsValid())
+			if (Frames.Num() > 0)
 			{
-				TSharedPtr<STimingView> TimingView = TimingWindow->GetTimingView();
+				TSharedPtr<STimingView> TimingView = GetTimingView();
 				if (TimingView.IsValid())
 				{
+					// Attempt to compute only from visible timelines.
+					TSet<uint32> VisibleTimelines;
+					TSet<uint32> VisibleCpuSamplingThreads;
 					TSharedPtr<FThreadTimingSharedState> ThreadSharedState = TimingView->GetThreadTimingSharedState();
 					if (ThreadSharedState.IsValid())
 					{
-						ThreadSharedState->GetVisibleTimelineIndexes(Timelines);
-						FFrameStatsHelper::ComputeFrameStatsForTimer(Frames, TimerSeries.GetTimerId(), Timelines);
-						bTimingViewExists = true;
+						ThreadSharedState->GetVisibleTimelineIndexes(VisibleTimelines);
+						ThreadSharedState->GetVisibleCpuSamplingThreads(VisibleCpuSamplingThreads);
 					}
-				}
-			}
 
-			if (!bTimingViewExists)
-			{
-				// Compute the stats for all timelines.
-				FFrameStatsHelper::ComputeFrameStatsForTimer(Frames, TimerSeries.GetTimerId());
-			}
+					// CPU & GPU timelines
+					if (VisibleTimelines.Num() > 0)
+					{
+						FFrameStatsHelper::ComputeFrameStatsForTimer(Frames, TimerId, VisibleTimelines);
+					}
 
-			uint64 CurrentIndex = StartIndex;
-			for (FFrameStatsCachedEvent& Event : Frames)
-			{
-				TraceServices::FFrame NewFrame;
-				NewFrame.StartTime = Event.FrameStartTime;
-				NewFrame.EndTime = Event.FrameStartTime + Event.Duration.load();
-				NewFrame.Index = CurrentIndex++;
-				Builder.AddFrame(NewFrame);
+					// CPU Sampling timelines
+					if (VisibleCpuSamplingThreads.Num() > 0)
+					{
+						const TraceServices::IStackSamplesProvider* StackSamplesProvider = TraceServices::ReadStackSamplesProvider(*Session.Get());
+						if (StackSamplesProvider)
+						{
+							for (uint32 ThreadId : VisibleCpuSamplingThreads)
+							{
+								TraceServices::FProviderReadScopeLock _(*StackSamplesProvider);
+								const TraceServices::ITimingProfilerProvider::Timeline& Timeline = *StackSamplesProvider->GetTimeline(ThreadId);
+								FFrameStatsHelper::ComputeFrameStatsForTimer(Frames, TimerId, Timeline);
+							}
+						}
+					}
+				}
+				else
+				{
+					// Compute the stats for all timelines.
+					FFrameStatsHelper::ComputeFrameStatsForTimer(Frames, TimerId);
+				}
+
+				uint64 CurrentIndex = StartIndex;
+				for (FFrameStatsCachedEvent& Event : Frames)
+				{
+					TraceServices::FFrame NewFrame;
+					NewFrame.StartTime = Event.FrameStartTime;
+					NewFrame.EndTime = Event.FrameStartTime + Event.Duration.load();
+					NewFrame.Index = CurrentIndex++;
+					Builder.AddFrame(NewFrame);
+				}
 			}
 
 			NumUpdatedFrames += Builder.GetNumAddedFrames();
@@ -556,39 +596,35 @@ void SFrameTrack::SelectFrameAtMousePosition(double X, double Y, bool JoinCurren
 
 	if (SampleRef.IsValid())
 	{
-		TSharedPtr<STimingProfilerWindow> Window = FTimingProfilerManager::Get()->GetProfilerWindow();
-		if (Window.IsValid())
+		TSharedPtr<STimingView> TimingView = GetTimingView();
+		if (TimingView.IsValid())
 		{
-			TSharedPtr<STimingView> TimingView = Window->GetTimingView();
-			if (TimingView.IsValid())
-			{
-				double StartTime = SampleRef.Sample->LargestFrameStartTime;
-				double Duration = SampleRef.Sample->LargestFrameDuration;
-
-				if (JoinCurrentSelection)
-				{
-					double EndTime = StartTime + Duration;
-					StartTime = FMath::Min(StartTime, TimingView->GetSelectionStartTime());
-					EndTime = FMath::Max(EndTime, TimingView->GetSelectionEndTime());
-					Duration = EndTime - StartTime;
-				}
+			double StartTime = SampleRef.Sample->LargestFrameStartTime;
+			double Duration = SampleRef.Sample->LargestFrameDuration;
 
-				TimingView->SetAutoScroll(false);
+			if (JoinCurrentSelection)
+			{
+				double EndTime = StartTime + Duration;
+				StartTime = FMath::Min(StartTime, TimingView->GetSelectionStartTime());
+				EndTime = FMath::Max(EndTime, TimingView->GetSelectionEndTime());
+				Duration = EndTime - StartTime;
+			}
 
-				if (bZoomTimingViewOnFrameSelection)
-				{
-					const double EndTime = FMath::Min(StartTime + Duration, TimingView->GetViewport().GetMaxValidTime());
-					const double AdjustedDuration = EndTime - StartTime;
-					TimingView->ZoomOnTimeInterval(StartTime - AdjustedDuration * 0.1, AdjustedDuration * 1.2);
-				}
-				else
-				{
-					TimingView->CenterOnTimeInterval(StartTime, Duration);
-				}
+			TimingView->SetAutoScroll(false);
 
-				TimingView->SelectTimeInterval(StartTime, Duration);
-				FSlateApplication::Get().SetKeyboardFocus(TimingView);
+			if (bZoomTimingViewOnFrameSelection)
+			{
+				const double EndTime = FMath::Min(StartTime + Duration, TimingView->GetViewport().GetMaxValidTime());
+				const double AdjustedDuration = EndTime - StartTime;
+				TimingView->ZoomOnTimeInterval(StartTime - AdjustedDuration * 0.1, AdjustedDuration * 1.2);
+			}
+			else
+			{
+				TimingView->CenterOnTimeInterval(StartTime, Duration);
 			}
+
+			TimingView->SelectTimeInterval(StartTime, Duration);
+			FSlateApplication::Get().SetKeyboardFocus(TimingView);
 		}
 	}
 }
@@ -644,17 +680,13 @@ int32 SFrameTrack::OnPaint(const FPaintArgs& Args, const FGeometry& AllottedGeom
 		TSharedPtr<FFrameTrackSeries> GameFrameSeries = FindSeries(ETraceFrameType::TraceFrameType_Game);
 		if (GameFrameSeries.IsValid())
 		{
-			TSharedPtr<STimingProfilerWindow> Window = FTimingProfilerManager::Get()->GetProfilerWindow();
-			if (Window)
+			TSharedPtr<STimingView> TimingView = GetTimingView();
+			if (TimingView.IsValid())
 			{
-				TSharedPtr<STimingView> TimingView = Window->GetTimingView();
-				if (TimingView)
-				{
-					// Highlight the area corresponding to viewport of Timing View.
-					const double StartTime = TimingView->GetViewport().GetStartTime();
-					const double EndTime = TimingView->GetViewport().GetEndTime();
-					Helper.DrawHighlightedInterval(*GameFrameSeries, StartTime, EndTime);
-				}
+				// Highlight the area corresponding to viewport of Timing View.
+				const double StartTime = TimingView->GetViewport().GetStartTime();
+				const double EndTime = TimingView->GetViewport().GetEndTime();
+				Helper.DrawHighlightedInterval(*GameFrameSeries, StartTime, EndTime);
 			}
 		}
 
@@ -2226,6 +2258,14 @@ uint32 SFrameTrack::GetNumSeriesForTimer(uint32 TimerId) const
 
 ////////////////////////////////////////////////////////////////////////////////////////////////////
 
+TSharedPtr<STimingView> SFrameTrack::GetTimingView() const
+{
+	TSharedPtr<STimingProfilerWindow> Window = FTimingProfilerManager::Get()->GetProfilerWindow();
+	return Window.IsValid() ? Window->GetTimingView() : nullptr;
+}
+
+////////////////////////////////////////////////////////////////////////////////////////////////////
+
 } // namespace UE::Insights::TimingProfiler
 
 #undef LOCTEXT_NAMESPACE
diff --git a/Engine/Source/Developer/TraceInsights/Private/Insights/TimingProfiler/Widgets/SFrameTrack.h b/Engine/Source/Developer/TraceInsights/Private/Insights/TimingProfiler/Widgets/SFrameTrack.h
--- a/Engine/Source/Developer/TraceInsights/Private/Insights/TimingProfiler/Widgets/SFrameTrack.h
+++ b/Engine/Source/Developer/TraceInsights/Private/Insights/TimingProfiler/Widgets/SFrameTrack.h
@@ -202,6 +202,10 @@ protected:
 
 	void ZoomHorizontally(const float Delta, const float X);
 
+	TSharedPtr<STimingView> GetTimingView() const;
+	void RegisterTimingViewCallbacks();
+	void UnregisterTimingViewCallbacks();
+
 protected:
 	/** The track's viewport. Encapsulates info about position and scale. */
 	FFrameTrackViewport Viewport;
@@ -279,7 +283,7 @@ protected:
 	/** Cursor type. */
 	ECursorType CursorType = ECursorType::Default;
 
-	STimingView* RegisteredTimingView = nullptr; // for pointer comparison only; do not dereferentiate!
+	TWeakPtr<STimingView> TimingViewWeak;
 	FDelegateHandle OnTrackVisibilityChangedHandle;
 	FDelegateHandle OnTrackAddedHandle;
 	FDelegateHandle OnTrackRemovedHandle;
diff --git a/Engine/Source/Developer/TraceInsights/Private/Insights/TimingProfiler/Widgets/SStatsView.cpp b/Engine/Source/Developer/TraceInsights/Private/Insights/TimingProfiler/Widgets/SStatsView.cpp
--- a/Engine/Source/Developer/TraceInsights/Private/Insights/TimingProfiler/Widgets/SStatsView.cpp
+++ b/Engine/Source/Developer/TraceInsights/Private/Insights/TimingProfiler/Widgets/SStatsView.cpp
@@ -37,6 +37,7 @@
 // TraceInsights
 #include "Insights/InsightsStyle.h"
 #include "Insights/Log.h"
+#include "Insights/TimingProfiler/GraphTracks/TimingGraphSeries.h"
 #include "Insights/TimingProfiler/GraphTracks/TimingGraphTrack.h"
 #include "Insights/TimingProfiler/TimingProfilerManager.h"
 #include "Insights/TimingProfiler/ViewModels/CounterAggregator.h"
@@ -2277,7 +2278,7 @@ void SStatsView::SetColorToNode(FStatsNode& StatsNode, FLinearColor CustomColor)
 
 	const uint32 CounterId = StatsNode.GetCounterId();
 	const FLinearColor BorderColor(CustomColor.R + 0.4f, CustomColor.G + 0.4f, CustomColor.B + 0.4f, 1.0f);
-	
+
 	TSharedPtr<FTimingGraphSeries> CounterSeries = GetTimingViewMainGraphTrack()->GetStatsCounterSeries(CounterId);
 	if (CounterSeries.IsValid())
 	{
@@ -2286,7 +2287,7 @@ void SStatsView::SetColorToNode(FStatsNode& StatsNode, FLinearColor CustomColor)
 }
 
 ////////////////////////////////////////////////////////////////////////////////////////////////////
-	
+
 void SStatsView::SetRandomColorToNode(FStatsNode& StatsNode) const
 {
 	uint64 Time = FPlatformTime::Cycles64();
@@ -2303,7 +2304,7 @@ void SStatsView::SetRandomColorToNode(FStatsNode& StatsNode) const
 
 	const uint32 CounterId = StatsNode.GetCounterId();
 	const FLinearColor BorderColor(CustomColor.R + 0.4f, CustomColor.G + 0.4f, CustomColor.B + 0.4f, 1.0f);
-	
+
 	TSharedPtr<FTimingGraphSeries> CounterSeries = GetTimingViewMainGraphTrack()->GetStatsCounterSeries(CounterId);
 	if (CounterSeries.IsValid())
 	{
diff --git a/Engine/Source/Developer/TraceInsights/Private/Insights/TimingProfiler/Widgets/STimerTreeView.cpp b/Engine/Source/Developer/TraceInsights/Private/Insights/TimingProfiler/Widgets/STimerTreeView.cpp
--- a/Engine/Source/Developer/TraceInsights/Private/Insights/TimingProfiler/Widgets/STimerTreeView.cpp
+++ b/Engine/Source/Developer/TraceInsights/Private/Insights/TimingProfiler/Widgets/STimerTreeView.cpp
@@ -24,6 +24,7 @@
 // TraceInsights
 #include "Insights/InsightsStyle.h"
 #include "Insights/Table/ViewModels/TableCommands.h"
+#include "Insights/TimingProfiler/GraphTracks/TimingGraphSeries.h"
 #include "Insights/TimingProfiler/GraphTracks/TimingGraphTrack.h"
 #include "Insights/TimingProfiler/TimingProfilerManager.h"
 #include "Insights/TimingProfiler/ViewModels/TimerButterflyAggregator.h"
diff --git a/Engine/Source/Developer/TraceInsights/Private/Insights/TimingProfiler/Widgets/STimersView.cpp b/Engine/Source/Developer/TraceInsights/Private/Insights/TimingProfiler/Widgets/STimersView.cpp
--- a/Engine/Source/Developer/TraceInsights/Private/Insights/TimingProfiler/Widgets/STimersView.cpp
+++ b/Engine/Source/Developer/TraceInsights/Private/Insights/TimingProfiler/Widgets/STimersView.cpp
@@ -44,6 +44,7 @@
 #include "Insights/InsightsStyle.h"
 #include "Insights/Log.h"
 #include "Insights/Table/ViewModels/TableCommands.h"
+#include "Insights/TimingProfiler/GraphTracks/TimingGraphSeries.h"
 #include "Insights/TimingProfiler/GraphTracks/TimingGraphTrack.h"
 #include "Insights/TimingProfiler/TimingProfilerManager.h"
 #include "Insights/TimingProfiler/Tracks/GpuTimingTrack.h"
diff --git a/Engine/Source/Developer/TraceInsights/Private/Insights/TimingProfiler/Widgets/STimingProfilerWindow.cpp b/Engine/Source/Developer/TraceInsights/Private/Insights/TimingProfiler/Widgets/STimingProfilerWindow.cpp
--- a/Engine/Source/Developer/TraceInsights/Private/Insights/TimingProfiler/Widgets/STimingProfilerWindow.cpp
+++ b/Engine/Source/Developer/TraceInsights/Private/Insights/TimingProfiler/Widgets/STimingProfilerWindow.cpp
@@ -11,6 +11,7 @@
 
 // TraceInsights
 #include "Insights/InsightsStyle.h"
+#include "Insights/TimingProfiler/GraphTracks/TimingGraphSeries.h"
 #include "Insights/TimingProfiler/GraphTracks/TimingGraphTrack.h"
 #include "Insights/TimingProfiler/TimingProfilerManager.h"
 #include "Insights/TimingProfiler/Widgets/SFrameTrack.h"
@@ -574,7 +575,7 @@ bool STimingProfilerWindow::GetTimerColor(uint32 TimerId, FLinearColor& OutColor
 		OutColor = *TimerColors.Find(TimerId);
 		return true;
 	}
-	
+
 	return false;
 }
 
@@ -584,16 +585,20 @@ void STimingProfilerWindow::SetTimerColor(uint32 TimerId, FLinearColor Color)
 {
 	TimerColors.Add(TimerId, Color);
 	const FLinearColor BorderColor(Color.R + 0.4f, Color.G + 0.4f, Color.B + 0.4f, 1.0f);
-	
+
 	if (TimingView.IsValid())
 	{
-		for (const TSharedPtr<FGraphSeries>& Series : TimingView->GetMainTimingGraphTrack()->GetSeries())
+		TSharedPtr<FTimingGraphTrack> GraphTrack = TimingView->GetMainTimingGraphTrack();
+		if (GraphTrack.IsValid())
 		{
-			if (Series->Is<FTimingGraphSeries>())
+			for (const TSharedPtr<FGraphSeries>& Series : GraphTrack->GetSeries())
 			{
-				if (Series->As<FTimingGraphSeries>().TimerId == TimerId)
+				if (Series->Is<FTimingGraphSeries>())
 				{
-					Series->SetColor(Color, BorderColor);
+					if (Series->As<FTimingGraphSeries>().IsTimer(TimerId))
+					{
+						Series->SetColor(Color, BorderColor);
+					}
 				}
 			}
 		}
