[Insights] Fixed cleanup for the rendering frame series added to the main graph track.

#rnx

diff --git a/Engine/Source/Developer/TraceInsights/Private/Insights/TimingProfiler/GraphTracks/TimingGraphTrack.cpp b/Engine/Source/Developer/TraceInsights/Private/Insights/TimingProfiler/GraphTracks/TimingGraphTrack.cpp
--- a/Engine/Source/Developer/TraceInsights/Private/Insights/TimingProfiler/GraphTracks/TimingGraphTrack.cpp
+++ b/Engine/Source/Developer/TraceInsights/Private/Insights/TimingProfiler/GraphTracks/TimingGraphTrack.cpp
@@ -69,7 +69,7 @@ FTimingGraphTrack::~FTimingGraphTrack()
 
 	if (RenderingFrameSeriesVisibilityHandle.IsValid())
 	{
-		TSharedPtr<FTimingGraphSeries> RenderingFramesSeries = GetFrameSeries(ETraceFrameType::TraceFrameType_Game);
+		TSharedPtr<FTimingGraphSeries> RenderingFramesSeries = GetFrameSeries(ETraceFrameType::TraceFrameType_Rendering);
 		if (RenderingFramesSeries.IsValid())
 		{
 			RenderingFramesSeries->VisibilityChangedDelegate.Remove(RenderingFrameSeriesVisibilityHandle);
