[Insights] GPU Timing Track: Fixed Fence Extended Lines still rendering when Fence Tracks are not visible

#rnx

diff --git a/Engine/Source/Developer/TraceInsights/Private/Insights/TimingProfiler/Tracks/GpuTimingTrack.cpp b/Engine/Source/Developer/TraceInsights/Private/Insights/TimingProfiler/Tracks/GpuTimingTrack.cpp
--- a/Engine/Source/Developer/TraceInsights/Private/Insights/TimingProfiler/Tracks/GpuTimingTrack.cpp
+++ b/Engine/Source/Developer/TraceInsights/Private/Insights/TimingProfiler/Tracks/GpuTimingTrack.cpp
@@ -431,13 +431,14 @@ void FGpuFencesTimingTrack::PostDraw(const ITimingTrackDrawContext& Context) con
 
 	constexpr double MaxViewportSizeForExtendedLines = 3.0;
 	const bool bExtentFenceLines = GetSharedState().AreGpuFencesExtendedLinesVisible() && Viewport.GetDuration() < MaxViewportSizeForExtendedLines;
+	const bool bFenceTracks = GetSharedState().AreGpuFencesTracksVisible();
 
 	float BoxY1 = GetPosY();
 	float BoxY2 = GetPosY() + GetHeight();
 	ETimingTrackLocation LocalLocation = GetLocation();
 
 	auto ParentTrackShared = GetParentTrack().Pin();
-	if (ParentTrack.IsValid() && bExtentFenceLines)
+	if (bFenceTracks && bExtentFenceLines && ParentTrack.IsValid())
 	{
 		// Substract the size of the child tracks above this track.
 		float TopY = GetPosY() - ParentTrackShared->GetPosY();
