VFS support in Unreal Insights.
Added a manual setting for mapping source files to a local repository.

diff --git a/Engine/Source/Developer/TraceInsights/Private/Insights/Common/SourceFilePathHelper.cpp b/Engine/Source/Developer/TraceInsights/Private/Insights/Common/SourceFilePathHelper.cpp
new file mode 100644
--- /dev/null
+++ b/Engine/Source/Developer/TraceInsights/Private/Insights/Common/SourceFilePathHelper.cpp
@@ -0,0 +1,192 @@
+// Copyright Epic Games, Inc. All Rights Reserved.
+
+#include "SourceFilePathHelper.h"
+
+#include "Insights/InsightsManager.h"
+#include "InsightsCore/Common/Log.h"
+#include "Misc/Paths.h"
+
+#define LOCTEXT_NAMESPACE "UE::Insights::FSourceFilePathHelper"
+
+namespace UE::Insights
+{
+
+////////////////////////////////////////////////////////////////////////////////////////////////////
+// FSourceFilePathHelper
+////////////////////////////////////////////////////////////////////////////////////////////////////
+
+FSourceFilePathHelper::FSourceFilePathHelper()
+{
+}
+
+////////////////////////////////////////////////////////////////////////////////////////////////////
+
+FSourceFilePathHelper::~FSourceFilePathHelper()
+{
+}
+
+////////////////////////////////////////////////////////////////////////////////////////////////////
+
+void FSourceFilePathHelper::InitVFSMapping(const FString& VFSPaths)
+{
+	TArray<FString> Entries;
+	VFSPaths.ParseIntoArray(Entries, TEXT(";"));
+
+	for (int32 Index = 0; Index < Entries.Num(); Index = Index + 2)
+	{
+		if (Index + 1 < Entries.Num())
+		{
+			VFSMappings.Add(Entries[Index], Entries[Index + 1]);
+		}
+		else
+		{
+			UE_LOG(LogInsights, Warning, TEXT("Detected unmatched VFS paths in the session:%s"), *VFSPaths);
+		}
+	}
+}
+
+////////////////////////////////////////////////////////////////////////////////////////////////////
+
+bool FSourceFilePathHelper::GetUsableFilePath(FString InPath, FString& OutPath)
+{
+	if (InPath.IsEmpty())
+	{
+		return false;
+	}
+
+	FPaths::NormalizeFilename(InPath);
+	// The goal of the LastAccessType checks is to minimise the number of calls to FPaths::FileExists.
+	if (LastAccessType == ELastAccessType::Direct && FPaths::FileExists(InPath))
+	{
+		OutPath = InPath;
+		LastAccessType = ELastAccessType::Direct;
+		return true;
+	}
+
+	const FString UserDefinedPath = GetLocalRootDirectoryPath();
+	if (!UserDefinedPath.IsEmpty() && LastAccessType == ELastAccessType::UserDefined)
+	{
+		for (auto& Pair : VFSMappings)
+		{
+			OutPath = InPath.Replace(*Pair.Value, *UserDefinedPath);
+			if (FPaths::FileExists(OutPath))
+			{
+				LastAccessType = ELastAccessType::UserDefined;
+				return true;
+			}
+
+			const TCHAR* Keywords[] = { TEXT("/Source/"), TEXT("/Plugins/"), TEXT("/Platforms/"),  TEXT("/Restricted/") };
+			for (const TCHAR* Keyword : Keywords)
+			{
+				// Try to reach the root path
+				int32 Index = InPath.Find(Keyword);
+				while (Index > 0 && (InPath[Index] == TEXT('/'))) --Index;
+				while (Index > 0 && (InPath[Index] != TEXT('/'))) --Index;
+				while (Index > 0 && (InPath[Index] == TEXT('/'))) --Index;
+				++Index;
+
+				if (Index <= InPath.Len())
+				{
+					FString RootPath = InPath.Left(Index);
+					OutPath = InPath.Replace(*RootPath, *UserDefinedPath);
+					if (FPaths::FileExists(OutPath))
+					{
+						LastAccessType = ELastAccessType::UserDefined;
+						return true;
+					}
+				}
+			}
+		}
+	}
+
+	for (auto& Pair : VFSMappings)
+	{
+		if (InPath.StartsWith(Pair.Key))
+		{
+			OutPath = InPath.Replace(*Pair.Key, *Pair.Value);
+			if (FPaths::FileExists(OutPath))
+			{
+				LastAccessType = ELastAccessType::VFS;
+				return true;
+			}
+			OutPath.Empty();
+		}
+	}
+
+	if (FPaths::FileExists(InPath))
+	{
+		OutPath = InPath;
+		LastAccessType = ELastAccessType::Direct;
+		return true;
+	}
+
+	if (!UserDefinedPath.IsEmpty())
+	{
+		for (auto& Pair : VFSMappings)
+		{
+			OutPath = InPath.Replace(*Pair.Value, *UserDefinedPath);
+			if (FPaths::FileExists(OutPath))
+			{
+				LastAccessType = ELastAccessType::UserDefined;
+				return true;
+			}
+		}
+
+		const TCHAR* Keywords[] = { TEXT("\\Source\\"), TEXT("\\Plugins\\"), TEXT("/Source/"), TEXT("/Plugins/") };
+		for (const TCHAR* Keyword : Keywords)
+		{
+			// Try to reach the root path
+			int32 Index = InPath.Find(Keyword);
+			while (Index > 0 && (InPath[Index] == TEXT('/') || InPath[Index] == TEXT('\\'))) --Index;
+			while (Index > 0 && (InPath[Index] != TEXT('/') && InPath[Index] != TEXT('\\'))) --Index;
+			while (Index > 0 && (InPath[Index] == TEXT('/') || InPath[Index] == TEXT('\\'))) --Index;
+			++Index;
+
+			if (Index <= InPath.Len())
+			{
+				FString RootPath = InPath.Left(Index);
+				OutPath = InPath.Replace(*RootPath, *UserDefinedPath);
+				if (FPaths::FileExists(OutPath))
+				{
+					LastAccessType = ELastAccessType::UserDefined;
+					return true;
+				}
+			}
+		}
+	}
+
+	OutPath = InPath;
+	return false;
+}
+
+////////////////////////////////////////////////////////////////////////////////////////////////////
+
+void FSourceFilePathHelper::SetUserDefinedPath(FString InPath)
+{
+	FPaths::NormalizeFilename(InPath);
+	FInsightsManager::Get()->GetSettings().SetAndSaveSourceFilesSearchPath(InPath);
+}
+
+///////////////////////////////////////////////////////////////////////////////////////////////////
+
+const FString& FSourceFilePathHelper::GetUserDefinedPath() const
+{
+	return FInsightsManager::Get()->GetSettings().GetSourceFilesSearchPath();
+}
+
+///////////////////////////////////////////////////////////////////////////////////////////////////
+
+FString FSourceFilePathHelper::GetLocalRootDirectoryPath() const
+{
+	FString Path = GetUserDefinedPath();
+	if (Path.IsEmpty())
+	{
+		return FPaths::RootDir();
+	}
+
+	return Path;
+}
+
+} // namespace UE::Insights
+
+#undef LOCTEXT_NAMESPACE
diff --git a/Engine/Source/Developer/TraceInsights/Private/Insights/Common/SourceFilePathHelper.h b/Engine/Source/Developer/TraceInsights/Private/Insights/Common/SourceFilePathHelper.h
new file mode 100644
--- /dev/null
+++ b/Engine/Source/Developer/TraceInsights/Private/Insights/Common/SourceFilePathHelper.h
@@ -0,0 +1,37 @@
+// Copyright Epic Games, Inc. All Rights Reserved.
+
+#pragma once
+
+#include "Containers/Map.h"
+#include "Containers/UnrealString.h"
+#include "CoreTypes.h"
+
+namespace UE::Insights
+{
+
+enum class ELastAccessType
+{
+	Unknown,
+	Direct,
+	VFS,
+	UserDefined,
+};
+
+class FSourceFilePathHelper
+{
+public:
+	FSourceFilePathHelper();
+	~FSourceFilePathHelper();
+
+	void InitVFSMapping(const FString& VFSPaths);
+	bool GetUsableFilePath(FString InPath, FString& OutPath);
+	const FString& GetUserDefinedPath() const;
+	void SetUserDefinedPath(FString InPath);
+	FString GetLocalRootDirectoryPath() const;
+
+private:
+	TMap<FString, FString> VFSMappings;
+	ELastAccessType LastAccessType = ELastAccessType::Unknown;
+};
+
+} // namespace UE::Insights
diff --git a/Engine/Source/Developer/TraceInsights/Private/Insights/InsightsManager.cpp b/Engine/Source/Developer/TraceInsights/Private/Insights/InsightsManager.cpp
--- a/Engine/Source/Developer/TraceInsights/Private/Insights/InsightsManager.cpp
+++ b/Engine/Source/Developer/TraceInsights/Private/Insights/InsightsManager.cpp
@@ -112,6 +112,7 @@ FInsightsManager::FInsightsManager(TSharedRef<TraceServices::IAnalysisService> I
 	, ModuleService(InTraceModuleService)
 	, CommandList(new FUICommandList())
 	, ActionManager(this)
+	, SourceFilePathHelper(MakeShared<FSourceFilePathHelper>())
 {
 }
 
@@ -312,7 +313,6 @@ bool FInsightsManager::Tick(float DeltaTime)
 
 	PollAnalysisInfo();
 
-#if !WITH_EDITOR
 	if (!bIsSessionInfoSet && Session.IsValid())
 	{
 		{
@@ -321,21 +321,28 @@ bool FInsightsManager::Tick(float DeltaTime)
 			if (DiagnosticsProvider && DiagnosticsProvider->IsSessionInfoAvailable())
 			{
 				bIsSessionInfoSet = true;
+				SourceFilePathHelper->InitVFSMapping(DiagnosticsProvider->GetSessionInfo().VFSPaths);
 				InstanceId = DiagnosticsProvider->GetSessionInfo().InstanceId;
+
+#if !WITH_EDITOR
 				TSharedPtr<SWidget> TraceControlPtr = TraceControl.Pin();
 				if (TraceControlPtr.IsValid())
 				{
 					FModuleManager::LoadModuleChecked<UE::TraceTools::ITraceToolsModule>("TraceTools").SetTraceControlWidgetInstanceId(TraceControlPtr.ToSharedRef(), InstanceId);
 				}
+#endif // !WITH_EDITOR
 			}
 		}
 
+#if !WITH_EDITOR
 		if (bIsSessionInfoSet)
 		{
 			UpdateAppTitle();
 		}
+#endif // !WITH_EDITOR
 	}
 
+#if !WITH_EDITOR
 	ISourceCodeAccessModule& SourceCodeAccessModule = FModuleManager::LoadModuleChecked<ISourceCodeAccessModule>("SourceCodeAccess");
 	ISourceCodeAccessor& SourceCodeAccessor = SourceCodeAccessModule.GetAccessor();
 	SourceCodeAccessor.Tick(DeltaTime);
diff --git a/Engine/Source/Developer/TraceInsights/Private/Insights/InsightsManager.h b/Engine/Source/Developer/TraceInsights/Private/Insights/InsightsManager.h
--- a/Engine/Source/Developer/TraceInsights/Private/Insights/InsightsManager.h
+++ b/Engine/Source/Developer/TraceInsights/Private/Insights/InsightsManager.h
@@ -16,6 +16,7 @@
 #include "InsightsCore/Common/Stopwatch.h"
 
 // TraceInsights
+#include "Insights/Common/SourceFilePathHelper.h"
 #include "Insights/IInsightsManager.h"
 #include "Insights/InsightsCommands.h"
 #include "Insights/InsightsSettings.h"
@@ -291,6 +292,9 @@ public:
 
 	void OpenTraceControlWindow();
 
+	FSourceFilePathHelper& GetSourceFilePathHelper() { return SourceFilePathHelper.Get(); }
+	FSourceFilePathHelper& GetSourceFilePathHelper() const { return SourceFilePathHelper.Get(); }
+
 	////////////////////////////////////////////////////////////////////////////////////////////////////
 	// SessionChangedEvent
 
@@ -328,7 +332,7 @@ private:
 
 	/** Extract messages from the session */
 	void PollAnalysisInfo();
-
+	
 	void OnSessionChanged();
 	void OnSessionAnalysisCompleted();
 	void OnSessionSymbolResolverCompleted();
@@ -434,6 +438,8 @@ private:
 
 	FGuid InstanceId;
 	TWeakPtr<SWidget> TraceControl;
+
+	TSharedRef<FSourceFilePathHelper> SourceFilePathHelper;
 };
 
 } // namespace UE::Insights
diff --git a/Engine/Source/Developer/TraceInsights/Private/Insights/InsightsSettings.cpp b/Engine/Source/Developer/TraceInsights/Private/Insights/InsightsSettings.cpp
--- a/Engine/Source/Developer/TraceInsights/Private/Insights/InsightsSettings.cpp
+++ b/Engine/Source/Developer/TraceInsights/Private/Insights/InsightsSettings.cpp
@@ -115,11 +115,6 @@ void FInsightsSettings::LoadFromConfig()
 		}
 	}
 
-	//////////////////////////////////////////////////
-	// [Insights.Symbols]
-
-	GConfig->GetArray(TEXT("Insights.Symbols"), TEXT("SymbolSearchPaths"), SymbolSearchPaths, SettingsIni);
-
 	//////////////////////////////////////////////////
 	// [Insights.TimingProfiler]
 
@@ -232,6 +227,11 @@ void FInsightsSettings::LoadFromConfig()
 	GConfig->GetBool(TEXT("Insights.TimingProfiler.TimersView"), TEXT("ShowCpuTimers"), bTimersViewShowCpuTimers, SettingsIni);
 	GConfig->GetBool(TEXT("Insights.TimingProfiler.TimersView"), TEXT("ShowCpuSamplingTimers"), bTimersViewShowCpuSamplingTimers, SettingsIni);
 	GConfig->GetBool(TEXT("Insights.TimingProfiler.TimersView"), TEXT("ShowZeroCountTimers"), bTimersViewShowZeroCountTimers, SettingsIni);
+
+	//////////////////////////////////////////////////
+	// [Insights]
+	GConfig->GetString(TEXT("Insights"), TEXT("SourceFilesSearchPath"), SourceFilesSearchPath, SettingsIni);
+	GConfig->GetArray(TEXT("Insights"), TEXT("SymbolSearchPaths"), SymbolSearchPaths, SettingsIni);
 }
 
 ////////////////////////////////////////////////////////////////////////////////////////////////////
@@ -291,11 +291,6 @@ void FInsightsSettings::SaveToConfig()
 		}
 	}
 
-	//////////////////////////////////////////////////
-	// [Insights.Symbols]
-
-	GConfig->SetArray(TEXT("Insights.Symbols"), TEXT("SymbolSearchPaths"), SymbolSearchPaths, SettingsIni);
-
 	//////////////////////////////////////////////////
 	// [Insights.TimingProfiler]
 
@@ -369,6 +364,11 @@ void FInsightsSettings::SaveToConfig()
 	GConfig->SetBool(TEXT("Insights.TimingProfiler.TimersView"), TEXT("ShowCpuSamplingTimers"), bTimersViewShowCpuSamplingTimers, SettingsIni);
 	GConfig->SetBool(TEXT("Insights.TimingProfiler.TimersView"), TEXT("ShowZeroCountTimers"), bTimersViewShowZeroCountTimers, SettingsIni);
 
+	//////////////////////////////////////////////////
+	// [Insights]
+	GConfig->SetString(TEXT("Insights"), TEXT("SourceFilesSearchPath"), *SourceFilesSearchPath, SettingsIni);
+	GConfig->SetArray(TEXT("Insights"), TEXT("SymbolSearchPaths"), SymbolSearchPaths, SettingsIni);
+
 	//////////////////////////////////////////////////
 
 	GConfig->Flush(false, SettingsIni);
diff --git a/Engine/Source/Developer/TraceInsights/Private/Insights/InsightsSettings.h b/Engine/Source/Developer/TraceInsights/Private/Insights/InsightsSettings.h
--- a/Engine/Source/Developer/TraceInsights/Private/Insights/InsightsSettings.h
+++ b/Engine/Source/Developer/TraceInsights/Private/Insights/InsightsSettings.h
@@ -329,13 +329,6 @@ public:
 		SaveToConfig();
 	}
 
-	//////////////////////////////////////////////////
-	// [Insights.Symbols]
-
-	const TArray<FString>& GetSymbolSearchPaths() const { return SymbolSearchPaths; }
-	void SetSymbolSearchPaths(const TArray<FString>& SearchPaths) { SymbolSearchPaths = SearchPaths; }
-	void SetAndSaveSymbolSearchPaths(const TArray<FString>& SearchPaths) { SET_AND_SAVE(SymbolSearchPaths, SearchPaths); }
-
 	//////////////////////////////////////////////////
 	// [Insights.TimingProfiler]
 
@@ -509,6 +502,17 @@ public:
 	void SetTimersViewShowZeroCountTimers(bool InValue) { bTimersViewShowZeroCountTimers = InValue; }
 	void SetAndSaveTimersViewShowZeroCountTimers(bool InValue) { SET_AND_SAVE(bTimersViewShowZeroCountTimers, InValue); }
 
+	//////////////////////////////////////////////////
+	// [Insights]
+
+	const FString& GetSourceFilesSearchPath() const { return SourceFilesSearchPath; }
+	void SetSourceFilesSearchPath(const FString& SearchPath) { SourceFilesSearchPath = SearchPath; }
+	void SetAndSaveSourceFilesSearchPath(const FString& SearchPath) { SET_AND_SAVE(SourceFilesSearchPath, SearchPath); }
+
+	const TArray<FString>& GetSymbolSearchPaths() const { return SymbolSearchPaths; }
+	void SetSymbolSearchPaths(const TArray<FString>& SearchPaths) { SymbolSearchPaths = SearchPaths; }
+	void SetAndSaveSymbolSearchPaths(const TArray<FString>& SearchPaths) { SET_AND_SAVE(SymbolSearchPaths, SearchPaths); }
+
 	//////////////////////////////////////////////////
 
 	#undef SET_AND_SAVE
@@ -535,12 +539,6 @@ private:
 	/** The default settings for a Timing View. */
 	inline static UE::Insights::FTimingViewSettings DefaultTimingViewSettings = { };
 
-	//////////////////////////////////////////////////
-	// [Insights.Symbols]
-
-	/** List of search paths to look for symbol files */
-	TArray<FString> SymbolSearchPaths;
-
 	//////////////////////////////////////////////////
 	// [Insights.TimingProfiler]
 
@@ -661,5 +659,14 @@ private:
 	/** If enabled, timers with no instances in the selected interval will still be displayed in the Timers View. */
 	bool bTimersViewShowZeroCountTimers = true;
 
+	//////////////////////////////////////////////////
+	// [Insights]
+
+	/** Search path for the open source files feature. */
+	FString SourceFilesSearchPath;
+
+	/** List of search paths to look for symbol files */
+	TArray<FString> SymbolSearchPaths;
+
 	//////////////////////////////////////////////////
 };
diff --git a/Engine/Source/Developer/TraceInsights/Private/Insights/MemoryProfiler/Widgets/SMemAllocTableTreeView.cpp b/Engine/Source/Developer/TraceInsights/Private/Insights/MemoryProfiler/Widgets/SMemAllocTableTreeView.cpp
--- a/Engine/Source/Developer/TraceInsights/Private/Insights/MemoryProfiler/Widgets/SMemAllocTableTreeView.cpp
+++ b/Engine/Source/Developer/TraceInsights/Private/Insights/MemoryProfiler/Widgets/SMemAllocTableTreeView.cpp
@@ -1432,19 +1432,22 @@ bool SMemAllocTableTreeView::BuildOpenSourceSubMenuItems(FMenuBuilder& MenuBuild
 		{
 			FText ItemLabel;
 			FText ItemToolTip;
+			FString UsablePath;
+			const bool bPathExists = FInsightsManager::Get()->GetSourceFilePathHelper().GetUsableFilePath(Frame->Symbol->File, UsablePath);
+
 			if (Frame->Symbol->GetResult() == TraceServices::ESymbolQueryResult::OK)
 			{
 				FText FileName;
 				constexpr int32 MaxFileNameLen = 120;
-				const int32 FileNameLen = FCString::Strlen(Frame->Symbol->File);
+				const int32 FileNameLen = UsablePath.Len();
 				if (FileNameLen > MaxFileNameLen)
 				{
-					FString FileNameStr = TEXT("...") + FString::ConstructFromPtrSize(Frame->Symbol->File + (FileNameLen - MaxFileNameLen), MaxFileNameLen);
+					FString FileNameStr = TEXT("...") + UsablePath.Right(MaxFileNameLen);
 					FileName = FText::FromString(FileNameStr);
 				}
 				else
 				{
-					FileName = FText::FromString(Frame->Symbol->File);
+					FileName = FText::FromString(UsablePath);
 				}
 
 				FText SymbolName;
@@ -1469,14 +1472,14 @@ bool SMemAllocTableTreeView::BuildOpenSourceSubMenuItems(FMenuBuilder& MenuBuild
 				{
 					ItemToolTip = FText::Format(LOCTEXT("ContextMenu_OpenSource_Desc_Fmt1", "Opens the source file of the selected callstack frame in {0}.\n{1} (line {2})\n\u2192 {3}"),
 						SourceCodeAccessor.GetNameText(),
-						FText::FromString(Frame->Symbol->File),
+						FText::FromString(UsablePath),
 						FText::AsNumber(Frame->Symbol->Line, &FNumberFormattingOptions::DefaultNoGrouping()),
 						FText::FromString(Frame->Symbol->Name));
 				}
 				else
 				{
 					ItemToolTip = FText::Format(LOCTEXT("ContextMenu_OpenSource_NoAccessor_Desc_Fmt1", "{0} (line {1})\n\u2192 {2}\nSource Code Accessor is not available."),
-						FText::FromString(Frame->Symbol->File),
+						FText::FromString(UsablePath),
 						FText::AsNumber(Frame->Symbol->Line, &FNumberFormattingOptions::DefaultNoGrouping()),
 						FText::FromString(Frame->Symbol->Name));
 				}
@@ -1499,18 +1502,17 @@ bool SMemAllocTableTreeView::BuildOpenSourceSubMenuItems(FMenuBuilder& MenuBuild
 				}
 			}
 
+			const bool bCanOpenSource = bPathExists && SourceCodeAccessor.CanAccessSourceCode();
 			switch (SourceSubMenuType)
 			{
 				case ESourceSubMenuType::OpenInSourceEditor:
 				{
-					const bool bCanOpenSource = SourceCodeAccessor.CanAccessSourceCode() && FPaths::FileExists(Frame->Symbol->File);
-
 					MenuBuilder.AddMenuEntry(
 						ItemLabel,
 						ItemToolTip,
 						FSlateIcon(),
 						FUIAction(
-							FExecuteAction::CreateSP(this, &SMemAllocTableTreeView::OpenSourceFileInIDE, Frame->Symbol->File, uint32(Frame->Symbol->Line)),
+							FExecuteAction::CreateSP(this, &SMemAllocTableTreeView::OpenSourceFileInIDE, UsablePath, uint32(Frame->Symbol->Line)),
 							FCanExecuteAction::CreateLambda([bCanOpenSource]() { return bCanOpenSource; })),
 						NAME_None,
 						EUserInterfaceActionType::Button);
@@ -1523,7 +1525,7 @@ bool SMemAllocTableTreeView::BuildOpenSourceSubMenuItems(FMenuBuilder& MenuBuild
 						LOCTEXT("ContextMenu_CopySourceCode_Desc", "Copy the source file path and line to the clipboard."),
 						FSlateIcon(),
 						FUIAction(
-							FExecuteAction::CreateSP(this, &SMemAllocTableTreeView::CopySourceFileToClipboard, Frame->Symbol->File, uint32(Frame->Symbol->Line)),
+							FExecuteAction::CreateSP(this, &SMemAllocTableTreeView::CopySourceFileToClipboard, UsablePath, uint32(Frame->Symbol->Line)),
 							FCanExecuteAction()),
 						NAME_None,
 						EUserInterfaceActionType::Button);
@@ -1575,27 +1577,25 @@ void SMemAllocTableTreeView::BuildOpenSourceSubMenu(FMenuBuilder& MenuBuilder, b
 
 ////////////////////////////////////////////////////////////////////////////////////////////////////
 
-void SMemAllocTableTreeView::OpenSourceFileInIDE(const TCHAR* InFile, uint32 Line) const
+void SMemAllocTableTreeView::OpenSourceFileInIDE(FString InFile, uint32 Line) const
 {
-	const FString File = InFile;
-
 	ISourceCodeAccessModule& SourceCodeAccessModule = FModuleManager::LoadModuleChecked<ISourceCodeAccessModule>("SourceCodeAccess");
-	if (FPaths::FileExists(File))
+	if (FPaths::FileExists(InFile))
 	{
 		ISourceCodeAccessor& SourceCodeAccessor = SourceCodeAccessModule.GetAccessor();
-		SourceCodeAccessor.OpenFileAtLine(File, Line);
+		SourceCodeAccessor.OpenFileAtLine(InFile, Line);
 	}
 	else
 	{
-		SourceCodeAccessModule.OnOpenFileFailed().Broadcast(File);
+		SourceCodeAccessModule.OnOpenFileFailed().Broadcast(InFile);
 	}
 }
 
 ////////////////////////////////////////////////////////////////////////////////////////////////////
 
-void SMemAllocTableTreeView::CopySourceFileToClipboard(const TCHAR* File, uint32 Line) const
+void SMemAllocTableTreeView::CopySourceFileToClipboard(FString File, uint32 Line) const
 {
-	FPlatformApplicationMisc::ClipboardCopy(*FString::Format(TEXT("{0} ({1})"), { File, Line }));
+	FPlatformApplicationMisc::ClipboardCopy(*FString::Format(TEXT("{0} ({1})"), { *File, Line }));
 }
 
 ////////////////////////////////////////////////////////////////////////////////////////////////////
@@ -1864,7 +1864,10 @@ void SMemAllocTableTreeView::OpenCallstackFrameSourceFileInIDE()
 			const TraceServices::FStackFrame* Frame = CallstackFrameNode.GetStackFrame();
 			if (Frame && Frame->Symbol && Frame->Symbol->File)
 			{
-				OpenSourceFileInIDE(Frame->Symbol->File, Frame->Symbol->Line);
+				FString UsablePath;
+				const bool bPathExists = FInsightsManager::Get()->GetSourceFilePathHelper().GetUsableFilePath(Frame->Symbol->File, UsablePath);
+
+				OpenSourceFileInIDE(UsablePath, Frame->Symbol->Line);
 			}
 		}
 	}
@@ -1883,7 +1886,9 @@ FText SMemAllocTableTreeView::GetSelectedCallstackFrameFileName() const
 			const TraceServices::FStackFrame* Frame = CallstackFrameNode.GetStackFrame();
 			if (Frame && Frame->Symbol && Frame->Symbol->File)
 			{
-				FString SourceFileAndLine = FString::Printf(TEXT("%s(%d)"), Frame->Symbol->File, Frame->Symbol->Line);
+				FString UsablePath;
+				const bool bPathExists = FInsightsManager::Get()->GetSourceFilePathHelper().GetUsableFilePath(Frame->Symbol->File, UsablePath);
+				FString SourceFileAndLine = FString::Printf(TEXT("%s(%d)"), *UsablePath, Frame->Symbol->Line);
 				return FText::FromString(SourceFileAndLine);
 			}
 			else
diff --git a/Engine/Source/Developer/TraceInsights/Private/Insights/MemoryProfiler/Widgets/SMemAllocTableTreeView.h b/Engine/Source/Developer/TraceInsights/Private/Insights/MemoryProfiler/Widgets/SMemAllocTableTreeView.h
--- a/Engine/Source/Developer/TraceInsights/Private/Insights/MemoryProfiler/Widgets/SMemAllocTableTreeView.h
+++ b/Engine/Source/Developer/TraceInsights/Private/Insights/MemoryProfiler/Widgets/SMemAllocTableTreeView.h
@@ -125,8 +125,8 @@ protected:
 	bool CanOpenCallstackFrameSourceFileInIDE() const;
 	void OpenCallstackFrameSourceFileInIDE();
 	FText GetSelectedCallstackFrameFileName() const;
-	void OpenSourceFileInIDE(const TCHAR* File, uint32 Line) const;
-	void CopySourceFileToClipboard(const TCHAR* File, uint32 Line) const;
+	void OpenSourceFileInIDE(FString File, uint32 Line) const;
+	void CopySourceFileToClipboard(FString File, uint32 Line) const;
 	void CopySelectedFilePathToClipboard() const;
 
 	void ExportMemorySnapshot() const;
diff --git a/Engine/Source/Developer/TraceInsights/Private/Insights/TaskGraphProfiler/Widgets/STaskTableTreeView.cpp b/Engine/Source/Developer/TraceInsights/Private/Insights/TaskGraphProfiler/Widgets/STaskTableTreeView.cpp
--- a/Engine/Source/Developer/TraceInsights/Private/Insights/TaskGraphProfiler/Widgets/STaskTableTreeView.cpp
+++ b/Engine/Source/Developer/TraceInsights/Private/Insights/TaskGraphProfiler/Widgets/STaskTableTreeView.cpp
@@ -116,6 +116,8 @@ void STaskTableTreeView::ExtendMenu(FMenuBuilder& MenuBuilder)
 		FString File;
 		uint32 Line = 0;
 		bool bIsValidSource = GetSourceFileAndLineForSelectedTask(File, Line);
+		FString UsablePath;
+		const bool bPathExists = FInsightsManager::Get()->GetSourceFilePathHelper().GetUsableFilePath(File, UsablePath);
 
 		ISourceCodeAccessModule& SourceCodeAccessModule = FModuleManager::LoadModuleChecked<ISourceCodeAccessModule>("SourceCodeAccess");
 		ISourceCodeAccessor& SourceCodeAccessor = SourceCodeAccessModule.GetAccessor();
@@ -130,7 +132,7 @@ void STaskTableTreeView::ExtendMenu(FMenuBuilder& MenuBuilder)
 			{
 				ItemToolTip = FText::Format(LOCTEXT("ContextMenu_OpenSource_Desc1", "Opens the source location where the selected task was launched, in {0}.\n{1} ({2})"),
 					SourceCodeAccessor.GetNameText(),
-					FText::FromString(File),
+					FText::FromString(UsablePath),
 					FText::AsNumber(Line, &FNumberFormattingOptions::DefaultNoGrouping()));
 			}
 			else
@@ -145,7 +147,7 @@ void STaskTableTreeView::ExtendMenu(FMenuBuilder& MenuBuilder)
 			if (bIsValidSource)
 			{
 				ItemToolTip = FText::Format(LOCTEXT("ContextMenu_OpenSourceNA_Desc1", "{0} ({1})\nSource Code Accessor is not available."),
-					FText::FromString(File),
+					FText::FromString(UsablePath),
 					FText::AsNumber(Line, &FNumberFormattingOptions::DefaultNoGrouping()));
 			}
 			else
@@ -859,15 +861,18 @@ void STaskTableTreeView::ContextMenu_OpenInIDE_Execute()
 		return;
 	}
 
+	FString UsablePath;
+	const bool bPathExists = FInsightsManager::Get()->GetSourceFilePathHelper().GetUsableFilePath(File, UsablePath);
 	ISourceCodeAccessModule& SourceCodeAccessModule = FModuleManager::LoadModuleChecked<ISourceCodeAccessModule>("SourceCodeAccess");
-	if (FPaths::FileExists(File))
+
+	if (bPathExists)
 	{
 		ISourceCodeAccessor& SourceCodeAccessor = SourceCodeAccessModule.GetAccessor();
-		SourceCodeAccessor.OpenFileAtLine(File, Line);
+		SourceCodeAccessor.OpenFileAtLine(UsablePath, Line);
 	}
 	else
 	{
-		SourceCodeAccessModule.OnOpenFileFailed().Broadcast(File);
+		SourceCodeAccessModule.OnOpenFileFailed().Broadcast(UsablePath);
 	}
 }
 
diff --git a/Engine/Source/Developer/TraceInsights/Private/Insights/TimingProfiler/Widgets/STimerTreeView.cpp b/Engine/Source/Developer/TraceInsights/Private/Insights/TimingProfiler/Widgets/STimerTreeView.cpp
--- a/Engine/Source/Developer/TraceInsights/Private/Insights/TimingProfiler/Widgets/STimerTreeView.cpp
+++ b/Engine/Source/Developer/TraceInsights/Private/Insights/TimingProfiler/Widgets/STimerTreeView.cpp
@@ -392,6 +392,8 @@ TSharedPtr<SWidget> STimerTreeView::TreeView_GetMenuContent()
 			{
 				bIsValidSource = SelectedNode->GetSourceFileAndLine(File, Line);
 			}
+			FString UsablePath;
+			const bool bPathExists = FInsightsManager::Get()->GetSourceFilePathHelper().GetUsableFilePath(File, UsablePath);
 
 			FText ItemLabel;
 			FText ItemToolTip;
@@ -403,7 +405,7 @@ TSharedPtr<SWidget> STimerTreeView::TreeView_GetMenuContent()
 				{
 					ItemToolTip = FText::Format(LOCTEXT("ContextMenu_OpenSource_Desc1", "Opens the source file of the selected timer in {0}.\n{1} ({2})"),
 						SourceCodeAccessor.GetNameText(),
-						FText::FromString(File),
+						FText::FromString(UsablePath),
 						FText::AsNumber(Line, &FNumberFormattingOptions::DefaultNoGrouping()));
 				}
 				else
@@ -418,7 +420,7 @@ TSharedPtr<SWidget> STimerTreeView::TreeView_GetMenuContent()
 				if (bIsValidSource)
 				{
 					ItemToolTip = FText::Format(LOCTEXT("ContextMenu_OpenSourceNA_Desc1", "{0} ({1})\nSource Code Accessor is not available."),
-						FText::FromString(File),
+						FText::FromString(UsablePath),
 						FText::AsNumber(Line, &FNumberFormattingOptions::DefaultNoGrouping()));
 				}
 				else
diff --git a/Engine/Source/Developer/TraceInsights/Private/Insights/TimingProfiler/Widgets/STimersView.cpp b/Engine/Source/Developer/TraceInsights/Private/Insights/TimingProfiler/Widgets/STimersView.cpp
--- a/Engine/Source/Developer/TraceInsights/Private/Insights/TimingProfiler/Widgets/STimersView.cpp
+++ b/Engine/Source/Developer/TraceInsights/Private/Insights/TimingProfiler/Widgets/STimersView.cpp
@@ -730,6 +730,8 @@ TSharedPtr<SWidget> STimersView::TreeView_GetMenuContent()
 			{
 				bIsValidSource = SelectedNode->GetSourceFileAndLine(File, Line);
 			}
+			FString UsablePath;
+			const bool bPathExists = FInsightsManager::Get()->GetSourceFilePathHelper().GetUsableFilePath(File, UsablePath);
 
 			FText ItemLabel;
 			FText ItemToolTip;
@@ -741,7 +743,7 @@ TSharedPtr<SWidget> STimersView::TreeView_GetMenuContent()
 				{
 					ItemToolTip = FText::Format(LOCTEXT("ContextMenu_OpenSource_Desc1", "Opens the source file of the selected timer in {0}.\n{1} ({2})"),
 						SourceCodeAccessor.GetNameText(),
-						FText::FromString(File),
+						FText::FromString(UsablePath),
 						FText::AsNumber(Line, &FNumberFormattingOptions::DefaultNoGrouping()));
 				}
 				else
@@ -756,7 +758,7 @@ TSharedPtr<SWidget> STimersView::TreeView_GetMenuContent()
 				if (bIsValidSource)
 				{
 					ItemToolTip = FText::Format(LOCTEXT("ContextMenu_OpenSourceNA_Desc1", "{0} ({1})\nSource Code Accessor is not available."),
-						FText::FromString(File),
+						FText::FromString(UsablePath),
 						FText::AsNumber(Line, &FNumberFormattingOptions::DefaultNoGrouping()));
 				}
 				else
@@ -1338,6 +1340,22 @@ void STimersView::TreeView_BuildSettingsMenu(FMenuBuilder& MenuBuilder)
 				OnThreadRatioCommited(Ratio);
 			})
 		);
+
+		FUIAction Action_SetSourceFilesSearchPath
+		(
+			FExecuteAction::CreateSP(this, &STimersView::SetSourceFileSearchPaths_Execute),
+			FCanExecuteAction::CreateLambda([]() { return true; })
+		);
+
+		MenuBuilder.AddMenuEntry
+		(
+			LOCTEXT("SetSourceFilesSearchPath_Fmt", "Source File Path"),
+			TAttribute<FText>::CreateSP(this, &STimersView::SetSourceFileSearchPaths_GetTooltip),
+			FSlateIcon(),
+			Action_SetSourceFilesSearchPath,
+			NAME_None,
+			EUserInterfaceActionType::Button
+		);
 	}
 	MenuBuilder.EndSection();
 }
@@ -3725,15 +3743,16 @@ void STimersView::OpenSourceFileInIDE(FTimerNodePtr InNode) const
 		return;
 	}
 
+	FString Path;
 	ISourceCodeAccessModule& SourceCodeAccessModule = FModuleManager::LoadModuleChecked<ISourceCodeAccessModule>("SourceCodeAccess");
-	if (FPaths::FileExists(File))
+	if (FInsightsManager::Get()->GetSourceFilePathHelper().GetUsableFilePath(File, Path))
 	{
 		ISourceCodeAccessor& SourceCodeAccessor = SourceCodeAccessModule.GetAccessor();
-		SourceCodeAccessor.OpenFileAtLine(File, Line);
+		SourceCodeAccessor.OpenFileAtLine(Path, Line);
 	}
 	else
 	{
-		SourceCodeAccessModule.OnOpenFileFailed().Broadcast(File);
+		SourceCodeAccessModule.OnOpenFileFailed().Broadcast(Path);
 	}
 }
 
@@ -4106,6 +4125,39 @@ bool STimersView::AreSelectedTimerNodesValid() const
 
 ////////////////////////////////////////////////////////////////////////////////////////////////////
 
+FText STimersView::SetSourceFileSearchPaths_GetTooltip() const
+{
+	const FString& CurrentPath = FInsightsManager::Get()->GetSourceFilePathHelper().GetLocalRootDirectoryPath();
+	return FText::Format(LOCTEXT("CurrentPathFmt", "Current path: {0}"), CurrentPath.Len() > 0 ? FText::FromString(CurrentPath) : LOCTEXT("None", "None"));
+}
+
+////////////////////////////////////////////////////////////////////////////////////////////////////
+
+void STimersView::SetSourceFileSearchPaths_Execute()
+{
+	if (IDesktopPlatform* DesktopPlatform = FDesktopPlatformModule::Get())
+	{
+		bool bHasSelected = false;
+		FString SelectedPath;
+		const FString& CurrentPath = FInsightsManager::Get()->GetSourceFilePathHelper().GetLocalRootDirectoryPath();
+		const FString Title = LOCTEXT("SetSourceFileSeachPaths", "Select repository root folder...").ToString();
+
+		bHasSelected = DesktopPlatform->OpenDirectoryDialog(
+			FSlateApplication::Get().FindBestParentWindowHandleForDialogs(AsShared()),
+			Title,
+			CurrentPath,
+			SelectedPath
+		);
+
+		if (bHasSelected)
+		{
+			FInsightsManager::Get()->GetSourceFilePathHelper().SetUserDefinedPath(SelectedPath);
+		}
+	}
+}
+
+////////////////////////////////////////////////////////////////////////////////////////////////////
+
 } // namespace UE::Insights::TimingProfiler
 
 #undef LOCTEXT_NAMESPACE
diff --git a/Engine/Source/Developer/TraceInsights/Private/Insights/TimingProfiler/Widgets/STimersView.h b/Engine/Source/Developer/TraceInsights/Private/Insights/TimingProfiler/Widgets/STimersView.h
--- a/Engine/Source/Developer/TraceInsights/Private/Insights/TimingProfiler/Widgets/STimersView.h
+++ b/Engine/Source/Developer/TraceInsights/Private/Insights/TimingProfiler/Widgets/STimersView.h
@@ -407,6 +407,9 @@ private:
 
 	void SetTimingViewFrameType();
 
+	FText SetSourceFileSearchPaths_GetTooltip() const;
+	void  SetSourceFileSearchPaths_Execute();
+
 	TSharedPtr<STimingView> GetTimingView() const;
 
 private:
diff --git a/Engine/Source/Developer/TraceInsights/Private/Insights/Widgets/SLogView.cpp b/Engine/Source/Developer/TraceInsights/Private/Insights/Widgets/SLogView.cpp
--- a/Engine/Source/Developer/TraceInsights/Private/Insights/Widgets/SLogView.cpp
+++ b/Engine/Source/Developer/TraceInsights/Private/Insights/Widgets/SLogView.cpp
@@ -1329,6 +1329,9 @@ TSharedPtr<SWidget> SLogView::ListView_GetContextMenu()
 				Line = Record.GetLine();
 			}
 
+			FString UsablePath;
+			const bool bPathExists = FInsightsManager::Get()->GetSourceFilePathHelper().GetUsableFilePath(File, UsablePath);
+
 			ISourceCodeAccessModule& SourceCodeAccessModule = FModuleManager::LoadModuleChecked<ISourceCodeAccessModule>("SourceCodeAccess");
 			ISourceCodeAccessor& SourceCodeAccessor = SourceCodeAccessModule.GetAccessor();
 			if (SourceCodeAccessor.CanAccessSourceCode())
@@ -1337,11 +1340,11 @@ TSharedPtr<SWidget> SLogView::ListView_GetContextMenu()
 					SourceCodeAccessor.GetNameText());
 
 				FText ItemToolTip;
-				if (!File.IsEmpty())
+				if (!UsablePath.IsEmpty())
 				{
 					ItemToolTip = FText::Format(LOCTEXT("ContextMenu_OpenSource_Desc1", "Opens the source file of the selected message in {0}.\n{1} ({2})"),
 						SourceCodeAccessor.GetNameText(),
-						FText::FromString(File),
+						FText::FromString(UsablePath),
 						FText::AsNumber(Line, &FNumberFormattingOptions::DefaultNoGrouping()));
 				}
 				else
@@ -1364,10 +1367,10 @@ TSharedPtr<SWidget> SLogView::ListView_GetContextMenu()
 				const FText ItemLabel = LOCTEXT("ContextMenu_OpenSource_NoAccessor", "Open Source");
 
 				FText ItemToolTip;
-				if (!File.IsEmpty())
+				if (!UsablePath.IsEmpty())
 				{
 					ItemToolTip = FText::Format(LOCTEXT("ContextMenu_OpenSource_NoAccessor_Desc1", "{0} ({1})\nSource Code Accessor is not available."),
-						FText::FromString(File),
+						FText::FromString(UsablePath),
 						FText::AsNumber(Line, &FNumberFormattingOptions::DefaultNoGrouping()));
 				}
 				else
@@ -2109,16 +2112,19 @@ void SLogView::OpenSource() const
 		FLogMessageRecord& Record = Cache.Get(SelectedLogMessage->GetIndex());
 		const FString File = Record.GetFile();
 		const uint32 Line = Record.GetLine();
+		
+		FString UsablePath;
+		const bool bPathExists = FInsightsManager::Get()->GetSourceFilePathHelper().GetUsableFilePath(File, UsablePath);
 
 		ISourceCodeAccessModule& SourceCodeAccessModule = FModuleManager::LoadModuleChecked<ISourceCodeAccessModule>("SourceCodeAccess");
-		if (FPaths::FileExists(File))
+		if (bPathExists)
 		{
 			ISourceCodeAccessor& SourceCodeAccessor = SourceCodeAccessModule.GetAccessor();
-			SourceCodeAccessor.OpenFileAtLine(File, Line);
+			SourceCodeAccessor.OpenFileAtLine(UsablePath, Line);
 		}
 		else
 		{
-			SourceCodeAccessModule.OnOpenFileFailed().Broadcast(File);
+			SourceCodeAccessModule.OnOpenFileFailed().Broadcast(UsablePath);
 		}
 	}
 }
diff --git a/Engine/Source/Developer/TraceServices/Private/Analyzers/DiagnosticsAnalysis.cpp b/Engine/Source/Developer/TraceServices/Private/Analyzers/DiagnosticsAnalysis.cpp
--- a/Engine/Source/Developer/TraceServices/Private/Analyzers/DiagnosticsAnalysis.cpp
+++ b/Engine/Source/Developer/TraceServices/Private/Analyzers/DiagnosticsAnalysis.cpp
@@ -86,6 +86,7 @@ bool FDiagnosticsAnalyzer::OnEvent(uint16 RouteId, EStyle Style, const FOnEventC
 		EventData.GetString("CommandLine", SessionInfo.CommandLine);
 		EventData.GetString("Branch", SessionInfo.Branch);
 		EventData.GetString("BuildVersion", SessionInfo.BuildVersion);
+		EventData.GetString("VFSPaths", SessionInfo.VFSPaths);
 		SessionInfo.Changelist = EventData.GetValue<uint32>("Changelist", 0);
 		SessionInfo.ConfigurationType = (EBuildConfiguration)EventData.GetValue<uint8>("ConfigurationType");
 		SessionInfo.TargetType = (EBuildTargetType)EventData.GetValue<uint8>("TargetType");
diff --git a/Engine/Source/Developer/TraceServices/Public/TraceServices/Model/Diagnostics.h b/Engine/Source/Developer/TraceServices/Public/TraceServices/Model/Diagnostics.h
--- a/Engine/Source/Developer/TraceServices/Public/TraceServices/Model/Diagnostics.h
+++ b/Engine/Source/Developer/TraceServices/Public/TraceServices/Model/Diagnostics.h
@@ -19,6 +19,7 @@ struct FSessionInfo
 	FString CommandLine;
 	FString Branch;
 	FString BuildVersion;
+	FString VFSPaths;
 	uint32 Changelist = 0;
 	EBuildConfiguration ConfigurationType = EBuildConfiguration::Unknown;
 	EBuildTargetType TargetType = EBuildTargetType::Unknown;
diff --git a/Engine/Source/Runtime/Core/Private/ProfilingDebugging/TraceAuxiliary.cpp b/Engine/Source/Runtime/Core/Private/ProfilingDebugging/TraceAuxiliary.cpp
--- a/Engine/Source/Runtime/Core/Private/ProfilingDebugging/TraceAuxiliary.cpp
+++ b/Engine/Source/Runtime/Core/Private/ProfilingDebugging/TraceAuxiliary.cpp
@@ -2,6 +2,8 @@
 
 #include "ProfilingDebugging/TraceAuxiliary.h"
 
+#include "BuildSettings.h"
+#include "Containers/UnrealString.h"
 #include "CoreGlobals.h"
 #include "Misc/ScopeRWLock.h"
 #include "ProfilingDebugging/CountersTrace.h"
@@ -1615,6 +1617,7 @@ UE_TRACE_EVENT_BEGIN(Diagnostics, Session2, NoSync|Important)
 	UE_TRACE_EVENT_FIELD(uint8, ConfigurationType)
 	UE_TRACE_EVENT_FIELD(uint8, TargetType)
 	UE_TRACE_EVENT_FIELD(uint32[], InstanceId)
+	UE_TRACE_EVENT_FIELD(UE::Trace::AnsiString, VFSPaths)
 UE_TRACE_EVENT_END()
 
 ////////////////////////////////////////////////////////////////////////////////
@@ -1943,6 +1946,8 @@ void FTraceAuxiliary::Initialize(const TCHAR* CommandLine)
 	const uint32 CommandLineLen = FCString::Strlen(CommandLine);
 	const uint32 BranchNameLen = FCString::Strlen(BranchName);
 	const uint32 BuildVersionLen = FCString::Strlen(BuildVersion);
+	const char* VFSPaths = BuildSettings::GetVfsPaths();
+	const uint32 VFSPathsLen = FCStringAnsi::Strlen(VFSPaths);
 	const uint32 DataSize =
 		(PlatformLen * sizeof(ANSICHAR)) +
 		(AppNameLen * sizeof(ANSICHAR)) +
@@ -1950,7 +1955,8 @@ void FTraceAuxiliary::Initialize(const TCHAR* CommandLine)
 		(CommandLineLen * sizeof(TCHAR)) +
 		(BranchNameLen * sizeof(TCHAR)) +
 		(BuildVersionLen * sizeof(TCHAR)) +
-		(InstanceIdSize * sizeof(uint32));
+		(InstanceIdSize * sizeof(uint32)) +
+		(VFSPathsLen * sizeof(ANSICHAR));
 
 	FGuid InstanceGuid = FApp::GetInstanceId();
 	uint32 InstanceId[InstanceIdSize];
@@ -1969,7 +1975,8 @@ void FTraceAuxiliary::Initialize(const TCHAR* CommandLine)
 		<< Session2.Changelist(BuildSettings::GetCurrentChangelist())
 		<< Session2.ConfigurationType(uint8(FApp::GetBuildConfiguration()))
 		<< Session2.TargetType(uint8(FApp::GetBuildTargetType()))
-		<< Session2.InstanceId(InstanceId, InstanceIdSize);
+		<< Session2.InstanceId(InstanceId, InstanceIdSize)
+		<< Session2.VFSPaths(VFSPaths, VFSPathsLen);
 
 	// Attempt to send trace data somewhere from the command line. It perhaps
 	// seems odd to do this before initializing Trace, but it is done this way
