[Insights] Fixes for open source files

diff --git a/Engine/Source/Developer/TraceInsights/Private/Insights/InsightsManager.h b/Engine/Source/Developer/TraceInsights/Private/Insights/InsightsManager.h
--- a/Engine/Source/Developer/TraceInsights/Private/Insights/InsightsManager.h
+++ b/Engine/Source/Developer/TraceInsights/Private/Insights/InsightsManager.h
@@ -293,7 +293,7 @@ public:
 	void OpenTraceControlWindow();
 
 	FSourceFilePathHelper& GetSourceFilePathHelper() { return SourceFilePathHelper.Get(); }
-	FSourceFilePathHelper& GetSourceFilePathHelper() const { return SourceFilePathHelper.Get(); }
+	const FSourceFilePathHelper& GetSourceFilePathHelper() const { return SourceFilePathHelper.Get(); }
 
 	////////////////////////////////////////////////////////////////////////////////////////////////////
 	// SessionChangedEvent
diff --git a/Engine/Source/Developer/TraceInsights/Private/Insights/TimingProfiler/Widgets/STimerTreeView.cpp b/Engine/Source/Developer/TraceInsights/Private/Insights/TimingProfiler/Widgets/STimerTreeView.cpp
--- a/Engine/Source/Developer/TraceInsights/Private/Insights/TimingProfiler/Widgets/STimerTreeView.cpp
+++ b/Engine/Source/Developer/TraceInsights/Private/Insights/TimingProfiler/Widgets/STimerTreeView.cpp
@@ -1919,15 +1919,17 @@ void STimerTreeView::OpenSourceFileInIDE(FTimerNodePtr InNode) const
 		return;
 	}
 
+	FString UsablePath;
+	const bool bPathExists = FInsightsManager::Get()->GetSourceFilePathHelper().GetUsableFilePath(File, UsablePath);
 	ISourceCodeAccessModule& SourceCodeAccessModule = FModuleManager::LoadModuleChecked<ISourceCodeAccessModule>("SourceCodeAccess");
-	if (FPaths::FileExists(File))
+	if (bPathExists)
 	{
 		ISourceCodeAccessor& SourceCodeAccessor = SourceCodeAccessModule.GetAccessor();
-		SourceCodeAccessor.OpenFileAtLine(File, Line);
+		SourceCodeAccessor.OpenFileAtLine(UsablePath, Line);
 	}
 	else
 	{
-		SourceCodeAccessModule.OnOpenFileFailed().Broadcast(File);
+		SourceCodeAccessModule.OnOpenFileFailed().Broadcast(UsablePath);
 	}
 }
 
