[Insights]
* Timing View / Export: Fixed export when -region is not correctly specified.
* Timing View / Export: Improved export error messages for all export commands.

diff --git a/Engine/Source/Developer/TraceInsights/Private/Insights/TimingProfiler/ViewModels/TimingExporter.cpp b/Engine/Source/Developer/TraceInsights/Private/Insights/TimingProfiler/ViewModels/TimingExporter.cpp
--- a/Engine/Source/Developer/TraceInsights/Private/Insights/TimingProfiler/ViewModels/TimingExporter.cpp
+++ b/Engine/Source/Developer/TraceInsights/Private/Insights/TimingProfiler/ViewModels/TimingExporter.cpp
@@ -45,7 +45,7 @@ FTimingExporter::~FTimingExporter()
 
 ////////////////////////////////////////////////////////////////////////////////////////////////////
 
-IFileHandle* FTimingExporter::OpenExportFile(const TCHAR* InFilename) const
+TUniquePtr<IFileHandle> FTimingExporter::OpenExportFile(const TCHAR* InFilename) const
 {
 	IPlatformFile& PlatformFile = FPlatformFileManager::Get().GetPlatformFile();
 
@@ -55,24 +55,19 @@ IFileHandle* FTimingExporter::OpenExportFile(const TCHAR* InFilename) const
 		PlatformFile.CreateDirectoryTree(*Path);
 	}
 
-	IFileHandle* ExportFileHandle = PlatformFile.OpenWrite(InFilename);
-
-	if (ExportFileHandle == nullptr)
-	{
-		Error(LOCTEXT("FailedToOpenFile", "Export failed. Failed to open file for write."));
-		return nullptr;
-	}
+	TUniquePtr<IFileHandle> ExportFileHandle{ PlatformFile.OpenWrite(InFilename) };
 
 	return ExportFileHandle;
 }
 
 ////////////////////////////////////////////////////////////////////////////////////////////////////
 
-void FTimingExporter::Error(const FText& InMessage) const
+void FTimingExporter::Error(const TCHAR* InTitle, const TCHAR* InMessage) const
 {
 	FName LogListingName = FTimingProfilerManager::Get()->GetLogListingName();
-	FMessageLog ReportMessageLog((LogListingName != NAME_None) ? LogListingName : TEXT("Other"));
-	ReportMessageLog.Error(InMessage);
+	FMessageLog ReportMessageLog((LogListingName != NAME_None) ? LogListingName : TEXT("Insights"));
+	ReportMessageLog.Error(FText::FromString(InTitle));
+	ReportMessageLog.Error(FText::FromString(InMessage));
 	ReportMessageLog.Notify();
 }
 
@@ -120,22 +115,25 @@ void FTimingExporter::FUtf8Writer::WriteStringBuilder(int32 CacheLen)
 
 int32 FTimingExporter::ExportThreadsAsText(const FString& Filename, FExportThreadsParams& Params) const
 {
+	const TCHAR* ErrorTitle = TEXT("Failed to export threads!");
+
 	if (Params.Columns != nullptr)
 	{
-		UE_LOG(TraceInsights, Error, TEXT("Custom list of columns is not yet supported!"));
+		Error(ErrorTitle, TEXT("Custom list of columns is not yet supported!"));
 		return -1;
 	}
 
 	FStopwatch Stopwatch;
 	Stopwatch.Start();
 
-	IFileHandle* ExportFileHandle = OpenExportFile(*Filename);
+	TUniquePtr<IFileHandle> ExportFileHandle = OpenExportFile(*Filename);
 	if (!ExportFileHandle)
 	{
+		Error(ErrorTitle, *FString::Printf(TEXT("Cannot write the export file (\"%s\")."), *Filename));
 		return -1;
 	}
 	bool bIsCSV = Filename.EndsWith(TEXT(".csv"));
-	FUtf8Writer Writer(ExportFileHandle, bIsCSV);
+	FUtf8Writer Writer(ExportFileHandle.Get(), bIsCSV);
 	FUtf8StringBuilder& StringBuilder = Writer.GetStringBuilder();
 
 	// Write header.
@@ -183,16 +181,17 @@ int32 FTimingExporter::ExportThreadsAsText(const FString& Filename, FExportThrea
 			const TraceServices::ITimingProfilerProvider* TimingProfilerProvider = TraceServices::ReadTimingProfilerProvider(Session);
 			if (TimingProfilerProvider)
 			{
-				TimingProfilerProvider->EnumerateGpuQueues([&](const TraceServices::FGpuQueueInfo& QueueInfo)
-				{
-					StringBuilder.Appendf(UTF8TEXT("%u"), GetNonCollidingId(QueueInfo.Id));
-					Writer.AppendSeparator();
-					StringBuilder.Append(QueueInfo.GetDisplayName());
-					Writer.AppendSeparator();
-					StringBuilder.Append(UTF8TEXT("GPU"));
-					Writer.AppendLineEnd();
-					++ThreadCount;
-				});
+				TimingProfilerProvider->EnumerateGpuQueues(
+					[&](const TraceServices::FGpuQueueInfo& QueueInfo)
+					{
+						StringBuilder.Appendf(UTF8TEXT("%u"), GetNonCollidingId(QueueInfo.Id));
+						Writer.AppendSeparator();
+						StringBuilder.Append(QueueInfo.GetDisplayName());
+						Writer.AppendSeparator();
+						StringBuilder.Append(UTF8TEXT("GPU"));
+						Writer.AppendLineEnd();
+						++ThreadCount;
+					});
 			}
 		}
 
@@ -218,8 +217,7 @@ int32 FTimingExporter::ExportThreadsAsText(const FString& Filename, FExportThrea
 
 	Writer.Flush();
 	ExportFileHandle->Flush();
-	delete ExportFileHandle;
-	ExportFileHandle = nullptr;
+	ExportFileHandle.Reset(nullptr);
 
 	Stopwatch.Stop();
 	const double TotalTime = Stopwatch.GetAccumulatedTime();
@@ -232,22 +230,25 @@ int32 FTimingExporter::ExportThreadsAsText(const FString& Filename, FExportThrea
 
 int32 FTimingExporter::ExportTimersAsText(const FString& Filename, FExportTimersParams& Params) const
 {
+	const TCHAR* ErrorTitle = TEXT("Failed to export timers!");
+
 	if (Params.Columns != nullptr)
 	{
-		UE_LOG(TraceInsights, Error, TEXT("Custom list of columns is not yet supported!"));
+		Error(ErrorTitle, TEXT("Custom list of columns is not yet supported!"));
 		return -1;
 	}
 
 	FStopwatch Stopwatch;
 	Stopwatch.Start();
 
-	IFileHandle* ExportFileHandle = OpenExportFile(*Filename);
+	TUniquePtr<IFileHandle> ExportFileHandle = OpenExportFile(*Filename);
 	if (!ExportFileHandle)
 	{
+		Error(ErrorTitle, *FString::Printf(TEXT("Cannot write the export file (\"%s\")."), *Filename));
 		return -1;
 	}
 	bool bIsCSV = Filename.EndsWith(TEXT(".csv"));
-	FUtf8Writer Writer(ExportFileHandle, bIsCSV);
+	FUtf8Writer Writer(ExportFileHandle.Get(), bIsCSV);
 	FUtf8StringBuilder& StringBuilder = Writer.GetStringBuilder();
 
 	// Write header.
@@ -264,59 +265,63 @@ int32 FTimingExporter::ExportTimersAsText(const FString& Filename, FExportTimers
 		Writer.AppendLineEnd();
 	}
 
+	const TraceServices::ITimingProfilerProvider* TimingProfilerProvider = TraceServices::ReadTimingProfilerProvider(Session);
+	if (!TimingProfilerProvider)
+	{
+		Error(ErrorTitle, TEXT("Unable to access TimingProfilerProvider."));
+		return -1;
+	}
+
 	uint32 TimerCount = 0;
 
 	// Write values.
-	if (TraceServices::ReadTimingProfilerProvider(Session))
 	{
 		TraceServices::FAnalysisSessionReadScope SessionReadScope(Session);
 
-		const TraceServices::ITimingProfilerProvider& TimingProfilerProvider = *TraceServices::ReadTimingProfilerProvider(Session);
-
-		const TraceServices::ITimingProfilerTimerReader* TimerReader;
-		TimingProfilerProvider.ReadTimers([&TimerReader](const TraceServices::ITimingProfilerTimerReader& Out) { TimerReader = &Out; });
-
-		TimerCount = TimerReader->GetTimerCount();
-		for (uint32 TimerIndex = 0; TimerIndex < TimerCount; ++TimerIndex)
-		{
-			const TraceServices::FTimingProfilerTimer& Timer = *(TimerReader->GetTimer(TimerIndex));
-			StringBuilder.Appendf(UTF8TEXT("%u"), Timer.Id);
-			Writer.AppendSeparator();
-			switch (Timer.Type)
-			{
-			case TraceServices::ETimingProfilerTimerType::CpuScope:
-				StringBuilder.Append(UTF8TEXT("CPU"));
-				break;
-			case TraceServices::ETimingProfilerTimerType::GpuScope:
-				StringBuilder.Append(UTF8TEXT("GPU"));
-				break;
-			case TraceServices::ETimingProfilerTimerType::VerseSampling:
-				StringBuilder.Append(UTF8TEXT("VerseSampling"));
-				break;
-			case TraceServices::ETimingProfilerTimerType::CpuSampling:
-				StringBuilder.Append(UTF8TEXT("CpuSampling"));
-				break;
-			default:
-				StringBuilder.Append(UTF8TEXT("Unknown"));
-				break;
-			}
-			Writer.AppendSeparator();
-			Writer.AppendString(Timer.Name);
-			Writer.AppendSeparator();
-			if (Timer.File)
+		TimingProfilerProvider->ReadTimers(
+			[&Writer, &StringBuilder, &TimerCount](const TraceServices::ITimingProfilerTimerReader& TimerReader)
 			{
-				StringBuilder.Append((const UTF8CHAR*)TCHAR_TO_UTF8(Timer.File));
-			}
-			Writer.AppendSeparator();
-			StringBuilder.Appendf(UTF8TEXT("%u"), Timer.Line);
-			Writer.AppendLineEnd();
-		}
+				TimerCount = TimerReader.GetTimerCount();
+				for (uint32 TimerIndex = 0; TimerIndex < TimerCount; ++TimerIndex)
+				{
+					const TraceServices::FTimingProfilerTimer& Timer = *(TimerReader.GetTimer(TimerIndex));
+					StringBuilder.Appendf(UTF8TEXT("%u"), Timer.Id);
+					Writer.AppendSeparator();
+					switch (Timer.Type)
+					{
+					case TraceServices::ETimingProfilerTimerType::GpuScope:
+						StringBuilder.Append(UTF8TEXT("GPU"));
+						break;
+					case TraceServices::ETimingProfilerTimerType::CpuScope:
+						StringBuilder.Append(UTF8TEXT("CPU"));
+						break;
+					case TraceServices::ETimingProfilerTimerType::CpuSampling:
+						StringBuilder.Append(UTF8TEXT("CpuSampling"));
+						break;
+					case TraceServices::ETimingProfilerTimerType::VerseSampling:
+						StringBuilder.Append(UTF8TEXT("VerseSampling"));
+						break;
+					default:
+						StringBuilder.Append(UTF8TEXT("Unknown"));
+						break;
+					}
+					Writer.AppendSeparator();
+					Writer.AppendString(Timer.Name);
+					Writer.AppendSeparator();
+					if (Timer.File)
+					{
+						StringBuilder.Append((const UTF8CHAR*)TCHAR_TO_UTF8(Timer.File));
+					}
+					Writer.AppendSeparator();
+					StringBuilder.Appendf(UTF8TEXT("%u"), Timer.Line);
+					Writer.AppendLineEnd();
+				}
+			});
 	}
 
 	Writer.Flush();
 	ExportFileHandle->Flush();
-	delete ExportFileHandle;
-	ExportFileHandle = nullptr;
+	ExportFileHandle.Reset(nullptr);
 
 	Stopwatch.Stop();
 	const double TotalTime = Stopwatch.GetAccumulatedTime();
@@ -597,12 +602,14 @@ int32 FTimingExporter::ExportTimingEvents_WriteEvents(FExportTimingEventsInterna
 
 int32 FTimingExporter::ExportTimingEventsAsTextByRegions(const FString& FilenamePattern, FExportTimingEventsParams& Params) const
 {
+	const TCHAR* ErrorTitle = TEXT("Failed to export timing events (by regions)!");
+
 	TMap<FString, FTimeRegionGroup> RegionGroups;
 	GetRegions(Params.Region, RegionGroups);
 
 	if (RegionGroups.Num() == 0)
 	{
-		UE_LOG(TraceInsights, Error, TEXT("Unable to find any region with name pattern '%s'."), *Params.Region);
+		Error(ErrorTitle, *FString::Printf(TEXT("Unable to find any region with name pattern '%s'."), *Params.Region));
 		return -1;
 	}
 
@@ -631,24 +638,33 @@ int32 FTimingExporter::ExportTimingEventsAsTextByRegions(const FString& Filename
 
 int32 FTimingExporter::ExportTimingEventsAsText(const FString& Filename, FExportTimingEventsParams& Params) const
 {
+	const TCHAR* ErrorTitle = TEXT("Failed to export timing events!");
+
 	if (!Params.Region.IsEmpty())
 	{
 		return ExportTimingEventsAsTextByRegions(Filename, Params);
 	}
 
+	if (Filename.Contains(TEXT("*")) || Filename.Contains(TEXT("{region}")))
+	{
+		Error(ErrorTitle, TEXT("Filename contains \"{region}\" pattern, but no valid \"-region\" parameter is specified."));
+		return -1;
+	}
+
 	ExportTimingEvents_InitColumns();
 	const TArray<FName>& Columns = Params.Columns ? *Params.Columns : ExportTimingEventsDefaultColumns;
 
 	FStopwatch Stopwatch;
 	Stopwatch.Start();
 
-	IFileHandle* ExportFileHandle = OpenExportFile(*Filename);
+	TUniquePtr<IFileHandle> ExportFileHandle = OpenExportFile(*Filename);
 	if (!ExportFileHandle)
 	{
+		Error(ErrorTitle, *FString::Printf(TEXT("Cannot write the export file (\"%s\")."), *Filename));
 		return -1;
 	}
 	bool bIsCSV = Filename.EndsWith(TEXT(".csv"));
-	FUtf8Writer Writer(ExportFileHandle, bIsCSV);
+	FUtf8Writer Writer(ExportFileHandle.Get(), bIsCSV);
 	FUtf8StringBuilder& StringBuilder = Writer.GetStringBuilder();
 
 	FExportTimingEventsInternalParams InternalParams = { *this, Params, Columns, Writer, 0 };
@@ -661,8 +677,7 @@ int32 FTimingExporter::ExportTimingEventsAsText(const FString& Filename, FExport
 
 	Writer.Flush();
 	ExportFileHandle->Flush();
-	delete ExportFileHandle;
-	ExportFileHandle = nullptr;
+	ExportFileHandle.Reset(nullptr);
 
 	Stopwatch.Stop();
 	const double TotalTime = Stopwatch.GetAccumulatedTime();
@@ -779,15 +794,17 @@ int32 FTimingExporter::EnumerateRegions(const TMap<FString, FTimeRegionGroup>& I
 
 			++ExportedRegionCount;
 
+			const TCHAR* ErrorTitle = TEXT("Failed to enumerate regions!");
+
 			// Avoid writing too many files...
 			if (IntervalIndex >= MaxIntervalsPerRegion)
 			{
-				UE_LOG(TraceInsights, Error, TEXT("Too many intervals for region '%s'! Exporting to separate file per interval for this region is not allowed to continue."), *KV.Key);
+				Error(ErrorTitle, *FString::Printf(TEXT("Too many intervals for region '%s'! Exporting to separate file per interval for this region is not allowed to continue."), *KV.Key));
 				break;
 			}
 			if (ExportedRegionCount >= MaxExportedRegions)
 			{
-				UE_LOG(TraceInsights, Error, TEXT("Too many regions! Exporting to separate file per region is not allowed to continue."));
+				Error(ErrorTitle, TEXT("Too many regions! Exporting to separate file per region is not allowed to continue."));
 				break;
 			}
 		}
@@ -805,12 +822,14 @@ int32 FTimingExporter::EnumerateRegions(const TMap<FString, FTimeRegionGroup>& I
 
 int32 FTimingExporter::ExportTimerStatisticsAsTextByRegions(const FString& FilenamePattern, FExportTimerStatisticsParams& Params) const
 {
+	const TCHAR* ErrorTitle = TEXT("Failed to export timing statistics (by regions)!");
+
 	TMap<FString, FTimeRegionGroup> RegionGroups;
 	GetRegions(Params.Region, RegionGroups);
 
 	if (RegionGroups.Num() == 0)
 	{
-		UE_LOG(TraceInsights, Error, TEXT("Unable to find any region with name pattern '%s'."), *Params.Region);
+		Error(ErrorTitle, *FString::Printf(TEXT("Unable to find any region with name pattern '%s'."), *Params.Region));
 		return -1;
 	}
 
@@ -821,7 +840,8 @@ int32 FTimingExporter::ExportTimerStatisticsAsTextByRegions(const FString& Filen
 	FExportTimerStatisticsParams RegionParams = Params;
 	RegionParams.Region.Reset();
 	int32 ExportedRegionCount = EnumerateRegions(RegionGroups, FilenamePattern,
-		[this, &RegionParams](const FString& Filename, const FString& RegionName, double IntervalStartTime, double IntervalEndTime)
+		[this, &RegionParams]
+		(const FString& Filename, const FString& RegionName, double IntervalStartTime, double IntervalEndTime)
 		{
 			RegionParams.IntervalStartTime = IntervalStartTime;
 			RegionParams.IntervalEndTime = IntervalEndTime;
@@ -877,22 +897,30 @@ TraceServices::FCreateAggregationParams::ESortOrder MapTimerStatisticsParams(FTi
 
 int32 FTimingExporter::ExportTimerStatisticsAsText(const FString& Filename, FExportTimerStatisticsParams& Params) const
 {
+	const TCHAR* ErrorTitle = TEXT("Failed to export timing statistics!");
+
 	if (!Params.Region.IsEmpty())
 	{
 		return ExportTimerStatisticsAsTextByRegions(Filename, Params);
 	}
 
+	if (Filename.Contains(TEXT("*")) || Filename.Contains(TEXT("{region}")))
+	{
+		Error(ErrorTitle, TEXT("Filename contains \"{region}\" pattern, but no valid \"-region\" parameter is specified."));
+		return -1;
+	}
+
+	const TraceServices::ITimingProfilerProvider* TimingProfilerProvider = TraceServices::ReadTimingProfilerProvider(Session);
+	if (!TimingProfilerProvider)
+	{
+		Error(ErrorTitle, TEXT("Unable to access TimingProfilerProvider."));
+		return -1;
+	}
+
 	TraceServices::ITable<TraceServices::FTimingProfilerAggregatedStats>* StatsTable;
 	{
 		TraceServices::FAnalysisSessionReadScope SessionReadScope(Session);
 
-		const TraceServices::ITimingProfilerProvider* TimingProfilerProvider = TraceServices::ReadTimingProfilerProvider(Session);
-		if (!TimingProfilerProvider)
-		{
-			UE_LOG(TraceInsights, Error, TEXT("Unable to access TimingProfilerProvider for ExportTimerStatisticsAsText"));
-			return -1;
-		}
-
 		TraceServices::FCreateAggregationParams CreateAggregationParams;
 		CreateAggregationParams.IntervalStart = Params.IntervalStartTime;
 		CreateAggregationParams.IntervalEnd = Params.IntervalEndTime;
@@ -900,7 +928,6 @@ int32 FTimingExporter::ExportTimerStatisticsAsText(const FString& Filename, FExp
 		CreateAggregationParams.bIncludeOldGpu2 = true;
 		CreateAggregationParams.GpuQueueFilter = [](uint32) { return true; };
 		CreateAggregationParams.CpuThreadFilter = Params.ThreadFilter;
-
 		CreateAggregationParams.SortBy = TimingExporter::MapTimerStatisticsParams(Params.SortBy);
 		CreateAggregationParams.SortOrder = TimingExporter::MapTimerStatisticsParams(Params.SortOrder);
 		CreateAggregationParams.TableEntryLimit = Params.MaxExportedEvents;
@@ -914,30 +941,31 @@ int32 FTimingExporter::ExportTimerStatisticsAsText(const FString& Filename, FExp
 
 	bool bSuccess = TraceServices::Table2Csv(*StatsTable, *Filename);
 
+	if (!bSuccess)
+	{
+		Error(ErrorTitle, *FString::Printf(TEXT("Failed to write the CSV file (\"%s\")!"), *Filename));
+		return -1;
+	}
+
 	Stopwatch.Stop();
 	const double TotalTime = Stopwatch.GetAccumulatedTime();
 
-	if (bSuccess)
-	{
-		UE_LOG(TraceInsights, Log, TEXT("Exported timing statistics to file in %.3fs (\"%s\")."), TotalTime, *Filename);
-	}
-	else
-	{
-		UE_LOG(TraceInsights, Error, TEXT("Failed to write the CSV file (\"%s\")!"), *Filename);
-	}
-	return bSuccess ? 1 : -2;
+	UE_LOG(TraceInsights, Log, TEXT("Exported timing statistics to file in %.3fs (\"%s\")."), TotalTime, *Filename);
+	return 1;
 }
 
 ////////////////////////////////////////////////////////////////////////////////////////////////////
 
 int32 FTimingExporter::ExportTimerCalleesByRegions(const FString& FilenamePattern, const FExportTimerCalleesParams& Params) const
 {
+	const TCHAR* ErrorTitle = TEXT("Failed to export timer callees (by regions)!");
+
 	TMap<FString, FTimeRegionGroup> RegionGroups;
 	GetRegions(Params.Region, RegionGroups);
 
 	if (RegionGroups.Num() == 0)
 	{
-		UE_LOG(TraceInsights, Error, TEXT("Unable to find any region with name pattern '%s'."), *Params.Region);
+		Error(ErrorTitle, *FString::Printf(TEXT("Unable to find any region with name pattern '%s'."), *Params.Region));
 		return -1;
 	}
 
@@ -966,21 +994,29 @@ int32 FTimingExporter::ExportTimerCalleesByRegions(const FString& FilenamePatter
 
 int32 FTimingExporter::ExportTimerCalleesAsText(const FString& Filename, const FExportTimerCalleesParams& Params) const
 {
+	const TCHAR* ErrorTitle = TEXT("Failed to export timer callees!");
+
 	if (!Params.Region.IsEmpty())
 	{
 		return ExportTimerCalleesByRegions(Filename, Params);
 	}
 
-	TUniquePtr<TraceServices::ITimingProfilerButterfly> Butterfly = [&Params, this]() -> TUniquePtr<TraceServices::ITimingProfilerButterfly>
+	if (Filename.Contains(TEXT("*")) || Filename.Contains(TEXT("{region}")))
 	{
-		TraceServices::FAnalysisSessionReadScope SessionReadScope(Session);
+		Error(ErrorTitle, TEXT("Filename contains \"{region}\" pattern, but no valid \"-region\" parameter is specified."));
+		return -1;
+	}
 
-		const TraceServices::ITimingProfilerProvider* TimingProfilerProvider = TraceServices::ReadTimingProfilerProvider(Session);
-		if (!TimingProfilerProvider)
-		{
-			UE_LOG(TraceInsights, Error, TEXT("Unable to access TimingProfilerProvider for ExportTimerCalleesAsText"));
-			return nullptr;
-		}
+	const TraceServices::ITimingProfilerProvider* TimingProfilerProvider = TraceServices::ReadTimingProfilerProvider(Session);
+	if (!TimingProfilerProvider)
+	{
+		Error(ErrorTitle, TEXT("Unable to access TimingProfilerProvider."));
+		return -1;
+	}
+
+	TUniquePtr<TraceServices::ITimingProfilerButterfly> Butterfly = [this, ErrorTitle, TimingProfilerProvider, &Params]() -> TUniquePtr<TraceServices::ITimingProfilerButterfly>
+	{
+		TraceServices::FAnalysisSessionReadScope SessionReadScope(Session);
 
 		TraceServices::FCreateButterflyParams ButterflyParams;
 		ButterflyParams.IntervalStart = Params.IntervalStartTime;
@@ -994,17 +1030,16 @@ int32 FTimingExporter::ExportTimerCalleesAsText(const FString& Filename, const F
 
 		return TUniquePtr<TraceServices::ITimingProfilerButterfly>{ TimingProfilerProvider->CreateButterfly(ButterflyParams) };
 	}();
-
 	if (!Butterfly)
 	{
-		UE_LOG(TraceInsights, Error, TEXT("ExportTimerCalleesAsText failed to create Butterfly for region %s [%.6f-%.6f]"), *Params.Region, Params.IntervalStartTime, Params.IntervalEndTime);
+		Error(ErrorTitle, *FString::Printf(TEXT("Cannot create Butterfly for region \"%s\" [%.6f-%.6f]"), *Params.Region, Params.IntervalStartTime, Params.IntervalEndTime));
 		return -1;
 	}
 
-	TUniquePtr<IFileHandle> ExportFileHandle{ OpenExportFile(*Filename) };
+	TUniquePtr<IFileHandle> ExportFileHandle = OpenExportFile(*Filename);
 	if (!ExportFileHandle)
 	{
-		UE_LOG(TraceInsights, Error, TEXT("ExportTimerCalleesAsText failed to open export file %s."), *Filename);
+		Error(ErrorTitle, *FString::Printf(TEXT("Cannot write the export file (\"%s\")."), *Filename));
 		return -1;
 	}
 
@@ -1225,21 +1260,21 @@ FTimingExporter::FTimingEventFilterFunc FTimingExporter::MakeTimingEventFilterBy
 	TMultiMap<FString, uint32> Timers;
 
 	// Iterate the GPU & CPU timers.
-	if (TraceServices::ReadTimingProfilerProvider(Session))
+	const TraceServices::ITimingProfilerProvider* TimingProfilerProvider = TraceServices::ReadTimingProfilerProvider(Session);
+	if (TimingProfilerProvider)
 	{
 		TraceServices::FAnalysisSessionReadScope SessionReadScope(Session);
 
-		const TraceServices::ITimingProfilerProvider& TimingProfilerProvider = *TraceServices::ReadTimingProfilerProvider(Session);
-
-		const TraceServices::ITimingProfilerTimerReader* TimerReader;
-		TimingProfilerProvider.ReadTimers([&TimerReader, &Timers](const TraceServices::ITimingProfilerTimerReader& Out) { TimerReader = &Out; });
-
-		uint32 TimerCount = TimerReader->GetTimerCount();
-		for (uint32 TimerIndex = 0; TimerIndex < TimerCount; ++TimerIndex)
-		{
-			const TraceServices::FTimingProfilerTimer& Timer = *(TimerReader->GetTimer(TimerIndex));
-			Timers.Add(FString(Timer.Name), Timer.Id);
-		}
+		TimingProfilerProvider->ReadTimers(
+			[&Timers](const TraceServices::ITimingProfilerTimerReader& TimerReader)
+			{
+				uint32 TimerCount = TimerReader.GetTimerCount();
+				for (uint32 TimerIndex = 0; TimerIndex < TimerCount; ++TimerIndex)
+				{
+					const TraceServices::FTimingProfilerTimer& Timer = *(TimerReader.GetTimer(TimerIndex));
+					Timers.Add(FString(Timer.Name), Timer.Id);
+				}
+			});
 	}
 
 	TArray<FString> Filter;
@@ -1345,22 +1380,25 @@ FTimingExporter::FTimingEventFilterFunc FTimingExporter::MakeTimingEventFilterBy
 
 int32 FTimingExporter::ExportCountersAsText(const FString& Filename, FExportCountersParams& Params) const
 {
+	const TCHAR* ErrorTitle = TEXT("Failed to export counters!");
+
 	if (Params.Columns != nullptr)
 	{
-		UE_LOG(TraceInsights, Error, TEXT("Custom list of columns is not yet supported!"));
+		Error(ErrorTitle, TEXT("Custom list of columns is not yet supported!"));
 		return -1;
 	}
 
 	FStopwatch Stopwatch;
 	Stopwatch.Start();
 
-	IFileHandle* ExportFileHandle = OpenExportFile(*Filename);
+	TUniquePtr<IFileHandle> ExportFileHandle = OpenExportFile(*Filename);
 	if (!ExportFileHandle)
 	{
+		Error(ErrorTitle, *FString::Printf(TEXT("Cannot write the export file (\"%s\")."), *Filename));
 		return -1;
 	}
 	bool bIsCSV = Filename.EndsWith(TEXT(".csv"));
-	FUtf8Writer Writer(ExportFileHandle, bIsCSV);
+	FUtf8Writer Writer(ExportFileHandle.Get(), bIsCSV);
 	FUtf8StringBuilder& StringBuilder = Writer.GetStringBuilder();
 
 	// Write header.
@@ -1408,8 +1446,7 @@ int32 FTimingExporter::ExportCountersAsText(const FString& Filename, FExportCoun
 
 	Writer.Flush();
 	ExportFileHandle->Flush();
-	delete ExportFileHandle;
-	ExportFileHandle = nullptr;
+	ExportFileHandle.Reset(nullptr);
 
 	Stopwatch.Stop();
 	const double TotalTime = Stopwatch.GetAccumulatedTime();
@@ -1422,12 +1459,14 @@ int32 FTimingExporter::ExportCountersAsText(const FString& Filename, FExportCoun
 
 int32 FTimingExporter::ExportCounterAsTextByRegions(const FString& FilenamePattern, uint32 CounterId, FExportCounterParams& Params) const
 {
+	const TCHAR* ErrorTitle = TEXT("Failed to export counter (by regions)!");
+
 	TMap<FString, FTimeRegionGroup> RegionGroups;
 	GetRegions(Params.Region, RegionGroups);
 
 	if (RegionGroups.Num() == 0)
 	{
-		UE_LOG(TraceInsights, Error, TEXT("Unable to find any region with name pattern '%s'."), *Params.Region);
+		Error(ErrorTitle, *FString::Printf(TEXT("Unable to find any region with name pattern '%s'."), *Params.Region));
 		return -1;
 	}
 
@@ -1438,7 +1477,8 @@ int32 FTimingExporter::ExportCounterAsTextByRegions(const FString& FilenamePatte
 	FExportCounterParams RegionParams = Params;
 	RegionParams.Region.Reset();
 	int32 ExportedRegionCount = EnumerateRegions(RegionGroups, FilenamePattern,
-		[this, CounterId, &RegionParams](const FString& Filename, const FString& RegionName, double IntervalStartTime, double IntervalEndTime)
+		[this, CounterId, &RegionParams]
+		(const FString& Filename, const FString& RegionName, double IntervalStartTime, double IntervalEndTime)
 		{
 			RegionParams.IntervalStartTime = IntervalStartTime;
 			RegionParams.IntervalEndTime = IntervalEndTime;
@@ -1456,9 +1496,11 @@ int32 FTimingExporter::ExportCounterAsTextByRegions(const FString& FilenamePatte
 
 int32 FTimingExporter::ExportCounterAsText(const FString& FilenamePattern, uint32 CounterId, FExportCounterParams& Params) const
 {
+	const TCHAR* ErrorTitle = TEXT("Failed to export counter!");
+
 	if (Params.Columns != nullptr)
 	{
-		UE_LOG(TraceInsights, Error, TEXT("Custom list of columns is not yet supported!"));
+		Error(ErrorTitle, TEXT("Custom list of columns is not yet supported!"));
 		return -1;
 	}
 
@@ -1485,7 +1527,7 @@ int32 FTimingExporter::ExportCounterAsText(const FString& FilenamePattern, uint3
 
 	if (CounterName.IsEmpty())
 	{
-		UE_LOG(TraceInsights, Error, TEXT("Invalid counter!"));
+		Error(ErrorTitle, TEXT("Invalid counter!"));
 		return -1;
 	}
 
@@ -1503,13 +1545,14 @@ int32 FTimingExporter::ExportCounterAsText(const FString& FilenamePattern, uint3
 		Filename.ReplaceInline(TEXT("{counter}"), *CounterFilename);
 	}
 
-	IFileHandle* ExportFileHandle = OpenExportFile(*Filename);
+	TUniquePtr<IFileHandle> ExportFileHandle = OpenExportFile(*Filename);
 	if (!ExportFileHandle)
 	{
+		Error(ErrorTitle, *FString::Printf(TEXT("Cannot write the export file (\"%s\")."), *Filename));
 		return -1;
 	}
 	bool bIsCSV = Filename.EndsWith(TEXT(".csv"));
-	FUtf8Writer Writer(ExportFileHandle, bIsCSV);
+	FUtf8Writer Writer(ExportFileHandle.Get(), bIsCSV);
 	FUtf8StringBuilder& StringBuilder = Writer.GetStringBuilder();
 
 	// Write header.
@@ -1626,8 +1669,7 @@ int32 FTimingExporter::ExportCounterAsText(const FString& FilenamePattern, uint3
 
 	Writer.Flush();
 	ExportFileHandle->Flush();
-	delete ExportFileHandle;
-	ExportFileHandle = nullptr;
+	ExportFileHandle.Reset(nullptr);
 
 	Stopwatch.Stop();
 	const double TotalTime = Stopwatch.GetAccumulatedTime();
diff --git a/Engine/Source/Developer/TraceInsights/Private/Insights/TimingProfiler/ViewModels/TimingExporter.h b/Engine/Source/Developer/TraceInsights/Private/Insights/TimingProfiler/ViewModels/TimingExporter.h
--- a/Engine/Source/Developer/TraceInsights/Private/Insights/TimingProfiler/ViewModels/TimingExporter.h
+++ b/Engine/Source/Developer/TraceInsights/Private/Insights/TimingProfiler/ViewModels/TimingExporter.h
@@ -9,6 +9,7 @@
 #include "Containers/Set.h"
 #include "HAL/Platform.h"
 #include "Templates/Function.h"
+#include "Templates/UniquePtr.h"
 #include "UObject/NameTypes.h"
 
 #include <limits>
@@ -326,8 +327,8 @@ public:
 	static FTimingEventFilterFunc MakeTimingEventFilterByTimersExclusive(const TSet<uint32>& ExcludedTimers);
 
 private:
-	IFileHandle* OpenExportFile(const TCHAR* InFilename) const;
-	void Error(const FText& InMessage) const;
+	TUniquePtr<IFileHandle> OpenExportFile(const TCHAR* InFilename) const;
+	void Error(const TCHAR* InTitle, const TCHAR* InMessage) const;
 	void ExportTimingEvents_InitColumns() const;
 	void ExportTimingEvents_WriteHeader(FExportTimingEventsInternalParams& Params) const;
 	int32 ExportTimingEvents_WriteEvents(FExportTimingEventsInternalParams& Params) const;
