[Insights] Timing View: Improved display and parsing for metadata attached to CPU and GPU timing events. Fixed case where metadata is a single string element. Added handling for some corner cases.

diff --git a/Engine/Source/Developer/TraceInsights/Private/Insights/TimingProfiler/Tracks/ThreadTimingTrack.cpp b/Engine/Source/Developer/TraceInsights/Private/Insights/TimingProfiler/Tracks/ThreadTimingTrack.cpp
--- a/Engine/Source/Developer/TraceInsights/Private/Insights/TimingProfiler/Tracks/ThreadTimingTrack.cpp
+++ b/Engine/Source/Developer/TraceInsights/Private/Insights/TimingProfiler/Tracks/ThreadTimingTrack.cpp
@@ -47,42 +47,50 @@ INSIGHTS_IMPLEMENT_RTTI(FThreadTimingTrackImpl)
 
 ////////////////////////////////////////////////////////////////////////////////////////////////////
 
-static void AppendMetadataToTooltip(FTooltipDrawState& Tooltip, const TraceServices::FMetadataSpec* MetadataSpec, TArrayView<const uint8>& Metadata)
+static void AppendMetadataToTooltip(FTooltipDrawState& InOutTooltip, const TraceServices::FMetadataSpec* InMetadataSpec, TArrayView<const uint8>& InMetadata)
 {
-	FMemoryReaderView MemoryReader(Metadata);
+	FMemoryReaderView MemoryReader(InMetadata);
 	FCborReader CborReader(&MemoryReader, ECborEndianness::StandardCompliant);
 	FCborContext Context;
 
-	if (MetadataSpec == nullptr)
+	if (InMetadataSpec == nullptr)
 	{
 		if (!CborReader.ReadNext(Context))
 		{
+			InOutTooltip.AddTitle(TEXTVIEW("<empty metadata>"), FLinearColor(1.0f, 0.5f, 0.5f, 1.0f));
+			return;
+		}
+
+		if (Context.MajorType() == ECborCode::TextString)
+		{
+			InOutTooltip.AddTitle(TEXTVIEW("Metadata..."), FLinearColor(0.3f, 0.3f, 0.3f, 1.0f));
+			InOutTooltip.AddNameValueTextLine(TEXTVIEW(""), Context.AsString());
 			return;
 		}
 
 		if (Context.MajorType() != ECborCode::Map)
 		{
+			InOutTooltip.AddTitle(TEXTVIEW("<invalid metadata>"), FLinearColor(1.0f, 0.5f, 0.5f, 1.0f));
 			return;
 		}
 	}
 
-	Tooltip.AddTitle(TEXT("Metadata:"));
+	InOutTooltip.AddTitle(TEXTVIEW("Metadata..."), FLinearColor(0.3f, 0.3f, 0.3f, 1.0f));
 
-	uint32 Index = 0;
-	while (true)
+	for (uint32 Index = 0; true; ++Index)
 	{
 		FString Key;
 
-		if (MetadataSpec)
+		if (InMetadataSpec)
 		{
-			if (Index < (uint32) MetadataSpec->FieldNames.Num())
+			if (Index < static_cast<uint32>(InMetadataSpec->FieldNames.Num()))
 			{
-				Key = MetadataSpec->FieldNames[Index];
-				++Index;
+				Key = InMetadataSpec->FieldNames[Index];
+				Key += TEXT(":");
 			}
 			else
 			{
-				Key = TEXT("UndefinedField");
+				Key = TEXT(""); // undefined field
 			}
 		}
 		else
@@ -94,10 +102,10 @@ static void AppendMetadataToTooltip(FTooltipDrawState& Tooltip, const TraceServi
 			}
 
 			Key = FString::ConstructFromPtrSize(Context.AsCString(), static_cast<int32>(Context.AsLength()));
+			Key += TEXT(":");
 		}
 
-		Key += TEXT(":");
-
+		// Read value
 		if (!CborReader.ReadNext(Context))
 		{
 			break;
@@ -105,57 +113,60 @@ static void AppendMetadataToTooltip(FTooltipDrawState& Tooltip, const TraceServi
 
 		switch (Context.MajorType())
 		{
+
 		case ECborCode::Int:
+		{
+			int64 Value = Context.AsInt();
+			FString ValueStr;
+			if (Value > 999'999'999LL)
 			{
-				int64 Value = Context.AsInt();
-				FString ValueStr;
-				if (Value > 999'999'999LL)
-				{
-					ValueStr += FString::Printf(TEXT("0x%llX"), Value);
-				}
-				else
-				{
-					ValueStr += FString::Printf(TEXT("%lld"), Value);
-				}
-				Tooltip.AddNameValueTextLine(Key, ValueStr);
-				continue;
+				ValueStr += FString::Printf(TEXT("0x%llX"), Value);
 			}
-		case ECborCode::Uint:
+			else
 			{
-				uint64 Value = Context.AsUInt();
-				FString ValueStr;
-				if (Value > 999'999'999ULL)
-				{
-					ValueStr = FString::Printf(TEXT("0x%llX"), Value);
-				}
-				else
-				{
-					ValueStr = FString::Printf(TEXT("%llu"), Value);
-				}
-				Tooltip.AddNameValueTextLine(Key, ValueStr);
-				continue;
+				ValueStr += FString::Printf(TEXT("%lld"), Value);
 			}
+			InOutTooltip.AddNameValueTextLine(Key, ValueStr);
+			continue;
+		}
 
-		case ECborCode::TextString:
+		case ECborCode::Uint:
+		{
+			uint64 Value = Context.AsUInt();
+			FString ValueStr;
+			if (Value > 999'999'999ULL)
 			{
-				FString Value = Context.AsString();
-				Tooltip.AddNameValueTextLine(Key, Value);
-				continue;
+				ValueStr = FString::Printf(TEXT("0x%llX"), Value);
 			}
-
-		case ECborCode::ByteString:
+			else
 			{
-				FString ValueStr = FString::ConstructFromPtrSize(Context.AsCString(), static_cast<int32>(Context.AsLength()));
-				Tooltip.AddNameValueTextLine(Key, ValueStr);
-				continue;
+				ValueStr = FString::Printf(TEXT("%llu"), Value);
 			}
+			InOutTooltip.AddNameValueTextLine(Key, ValueStr);
+			continue;
+		}
+
+		case ECborCode::TextString:
+		{
+			FString Value = Context.AsString();
+			InOutTooltip.AddNameValueTextLine(Key, Value);
+			continue;
 		}
 
+		case ECborCode::ByteString:
+		{
+			FString ValueStr = FString::ConstructFromPtrSize(Context.AsCString(), static_cast<int32>(Context.AsLength()));
+			InOutTooltip.AddNameValueTextLine(Key, ValueStr);
+			continue;
+		}
+
+		} // switch
+
 		if (Context.RawCode() == (ECborCode::Prim|ECborCode::Value_4Bytes))
 		{
 			float Value = Context.AsFloat();
 			FString ValueStr = FString::Printf(TEXT("%f"), Value);
-			Tooltip.AddNameValueTextLine(Key, ValueStr);
+			InOutTooltip.AddNameValueTextLine(Key, ValueStr);
 			continue;
 		}
 
@@ -163,22 +174,24 @@ static void AppendMetadataToTooltip(FTooltipDrawState& Tooltip, const TraceServi
 		{
 			double Value = Context.AsDouble();
 			FString ValueStr = FString::Printf(TEXT("%g"), Value);
-			Tooltip.AddNameValueTextLine(Key, ValueStr);
+			InOutTooltip.AddNameValueTextLine(Key, ValueStr);
 			continue;
 		}
 
 		if (Context.RawCode() == (ECborCode::Prim | ECborCode::False))
 		{
-			Tooltip.AddNameValueTextLine(Key, FString(TEXT("false")));
+			InOutTooltip.AddNameValueTextLine(Key, TEXTVIEW("false"));
 			continue;
 		}
 
 		if (Context.RawCode() == (ECborCode::Prim | ECborCode::True))
 		{
-			Tooltip.AddNameValueTextLine(Key, FString(TEXT("true")));
+			InOutTooltip.AddNameValueTextLine(Key, TEXTVIEW("true"));
 			continue;
 		}
 
+		InOutTooltip.AddNameValueTextLine(Key, TEXTVIEW("???")); // unknown field type
+
 		if (Context.IsFiniteContainer())
 		{
 			CborReader.SkipContainer(ECborCode::Array);
@@ -188,20 +201,32 @@ static void AppendMetadataToTooltip(FTooltipDrawState& Tooltip, const TraceServi
 
 ////////////////////////////////////////////////////////////////////////////////////////////////////
 
-static void AppendMetadataToString(FString& Str, TArrayView<const uint8>& Metadata)
+static void AppendMetadataToString(FString& InOutStr, TArrayView<const uint8>& InMetadata)
 {
-	FMemoryReaderView MemoryReader(Metadata);
+	FMemoryReaderView MemoryReader(InMetadata);
 	FCborReader CborReader(&MemoryReader, ECborEndianness::StandardCompliant);
 	FCborContext Context;
 
-	if (!CborReader.ReadNext(Context) || Context.MajorType() != ECborCode::Map)
+	if (!CborReader.ReadNext(Context))
 	{
+		InOutStr += TEXTVIEW(" - <empty metadata>");
 		return;
 	}
 
-	bool bFirst = true;
+	if (Context.MajorType() == ECborCode::TextString)
+	{
+		InOutStr += TEXTVIEW(" - ");
+		InOutStr += Context.AsString();
+		return;
+	}
 
-	while (true)
+	if (Context.MajorType() != ECborCode::Map)
+	{
+		InOutStr += TEXTVIEW(" - <invalid metadata>");
+		return;
+	}
+
+	for (uint32 Index = 0; true; ++Index)
 	{
 		// Read key
 		if (!CborReader.ReadNext(Context) || !Context.IsString())
@@ -209,20 +234,20 @@ static void AppendMetadataToString(FString& Str, TArrayView<const uint8>& Metada
 			break;
 		}
 
-		if (bFirst)
+		if (Index == 0)
 		{
-			bFirst = false;
-			Str += TEXT(" - ");
+			InOutStr += TEXTVIEW(" - ");
 		}
 		else
 		{
-			Str += TEXT(", ");
+			InOutStr += TEXTVIEW(", ");
 		}
 
 		//FString Key(Context.AsLength(), Context.AsCString());
-		//Str += Key;
-		//Str += TEXT(":");
+		//InOutStr += Key;
+		//InOutStr += TEXTVIEW(":");
 
+		// Read value
 		if (!CborReader.ReadNext(Context))
 		{
 			break;
@@ -230,72 +255,77 @@ static void AppendMetadataToString(FString& Str, TArrayView<const uint8>& Metada
 
 		switch (Context.MajorType())
 		{
+
 		case ECborCode::Int:
 		{
 			int64 Value = Context.AsInt();
 			if (Value > 999'999'999LL)
 			{
-				Str += FString::Printf(TEXT("0x%llX"), Value);
+				InOutStr += FString::Printf(TEXT("0x%llX"), Value);
 			}
 			else
 			{
-				Str += FString::Printf(TEXT("%lld"), Value);
+				InOutStr += FString::Printf(TEXT("%lld"), Value);
 			}
 			continue;
 		}
+
 		case ECborCode::Uint:
 		{
 			uint64 Value = Context.AsUInt();
 			if (Value > 999'999'999ULL)
 			{
-				Str += FString::Printf(TEXT("0x%llX"), Value);
+				InOutStr += FString::Printf(TEXT("0x%llX"), Value);
 			}
 			else
 			{
-				Str += FString::Printf(TEXT("%llu"), Value);
+				InOutStr += FString::Printf(TEXT("%llu"), Value);
 			}
 			continue;
 		}
 
 		case ECborCode::TextString:
 		{
-			Str += Context.AsString();
+			InOutStr += Context.AsString();
 			continue;
 		}
 
 		case ECborCode::ByteString:
 		{
-			Str.AppendChars(Context.AsCString(), static_cast<int32>(Context.AsLength()));
+			InOutStr.AppendChars(Context.AsCString(), static_cast<int32>(Context.AsLength()));
 			continue;
 		}
-		}
+
+		} // switch
 
 		if (Context.RawCode() == (ECborCode::Prim | ECborCode::Value_4Bytes))
 		{
 			float Value = Context.AsFloat();
-			Str += FString::Printf(TEXT("%f"), Value);
+			InOutStr += FString::Printf(TEXT("%f"), Value);
 			continue;
 		}
 
 		if (Context.RawCode() == (ECborCode::Prim | ECborCode::Value_8Bytes))
 		{
 			double Value = Context.AsDouble();
-			Str += FString::Printf(TEXT("%g"), Value);
+			InOutStr += FString::Printf(TEXT("%g"), Value);
 			continue;
 		}
 
 		if (Context.RawCode() == (ECborCode::Prim | ECborCode::False))
 		{
-			Str += TEXT("false");
+			InOutStr += TEXTVIEW("false");
 			continue;
 		}
 
 		if (Context.RawCode() == (ECborCode::Prim | ECborCode::True))
 		{
-			Str += TEXT("true");
+			InOutStr += TEXTVIEW("true");
 			continue;
 		}
 
+		InOutStr += TEXTVIEW("???"); // unknown field type
+
 		if (Context.IsFiniteContainer())
 		{
 			CborReader.SkipContainer(ECborCode::Array);
@@ -305,27 +335,25 @@ static void AppendMetadataToString(FString& Str, TArrayView<const uint8>& Metada
 
 ////////////////////////////////////////////////////////////////////////////////////////////////////
 
-static void AppendMetadataToString(FString& Str, const TraceServices::FMetadataSpec* MetadataSpec, TArrayView<const uint8>& Metadata)
+static void AppendMetadataToString(FString& InOutStr, const FString& InFormat, TArrayView<const uint8>& InMetadata)
 {
-	if (MetadataSpec == nullptr)
-	{
-		return;
-	}
-
-	FMemoryReaderView MemoryReader(Metadata);
+	FMemoryReaderView MemoryReader(InMetadata);
 	FCborReader CborReader(&MemoryReader, ECborEndianness::StandardCompliant);
 	FCborContext Context;
 
-	FString Format;
-	Format = MetadataSpec->Format;
+	FString Format = InFormat;
 
-	const FString Specifiers = TEXT("diuoxXfFeEgGaAcspn");
+	Format.TrimStartInline();
+	if (Format.StartsWith(TEXT("- ")))
+	{
+		Format.RightChopInline(2);
+	}
 
+	const FString Specifiers = TEXT("diuoxXfFeEgGaAcspn");
 	auto GetNextFormatSection = [&Format, &Specifiers]()
 	{
-		int Index = 0;
 		bool bIsInFormatSpecifier = false;
-		while (Index < Format.Len())
+		for (int32 Index = 0; Index < Format.Len(); ++Index)
 		{
 			if (bIsInFormatSpecifier)
 			{
@@ -342,8 +370,6 @@ static void AppendMetadataToString(FString& Str, const TraceServices::FMetadataS
 			{
 				bIsInFormatSpecifier = !bIsInFormatSpecifier;
 			}
-
-			++Index;
 		}
 
 		FString Copy = Format;
@@ -351,68 +377,62 @@ static void AppendMetadataToString(FString& Str, const TraceServices::FMetadataS
 		return Copy;
 	};
 
-	bool bFirst = true;
-
-	while (!Format.IsEmpty())
+	for (uint32 Index = 0; !Format.IsEmpty(); ++Index)
 	{
+		// Read key
 		if (!CborReader.ReadNext(Context))
 		{
 			break;
 		}
 
-		if (bFirst)
+		if (Index == 0)
 		{
-			bFirst = false;
-			Str += TEXT(" - ");
-		}
-		else
-		{
-			if (Format.IsEmpty())
-			{
-				Str += TEXT(", ");
-			}
+			InOutStr += TEXTVIEW(" - ");
 		}
 
-		constexpr int MaxLength = 256;
+		constexpr int MaxLength = 1024;
 		TCHAR Data[MaxLength];
 
-		auto AddValueToName = [&Str](TCHAR* Dest, int MaxLength, const FString& Format, auto Value)
+		auto AddValueToName = [&InOutStr](TCHAR* Dest, int MaxLength, const FString& Format, auto Value)
 		{
 			int32 Result = FCString::Snprintf(Dest, MaxLength, reinterpret_cast<TCHAR const(&)[1]>(**Format), Value);
 			if (Result > 0)
 			{
-				Str.Append(Dest);
+				InOutStr.Append(Dest);
 			}
 		};
 
 		switch (Context.MajorType())
 		{
+
 		case ECborCode::Int:
-			{
-				int64 Value = Context.AsInt();
-				AddValueToName(Data, MaxLength, GetNextFormatSection(), Value);
-				continue;
-			}
+		{
+			int64 Value = Context.AsInt();
+			AddValueToName(Data, MaxLength, GetNextFormatSection(), Value);
+			continue;
+		}
+
 		case ECborCode::Uint:
-			{
-				uint64 Value = Context.AsUInt();
-				AddValueToName(Data, MaxLength, GetNextFormatSection(), Value);
-				continue;
-			}
+		{
+			uint64 Value = Context.AsUInt();
+			AddValueToName(Data, MaxLength, GetNextFormatSection(), Value);
+			continue;
+		}
 
 		case ECborCode::TextString:
-			{
-				AddValueToName(Data, MaxLength, GetNextFormatSection(), *Context.AsString());
-				continue;
-			}
+		{
+			AddValueToName(Data, MaxLength, GetNextFormatSection(), *Context.AsString());
+			continue;
+		}
 
 		case ECborCode::ByteString:
-			{
-				AddValueToName(Data, FMath::Min(MaxLength, (int)Context.AsLength()), GetNextFormatSection(), Context.AsCString());
-				continue;
-			}
+		{
+			AddValueToName(Data, FMath::Min(MaxLength, (int)Context.AsLength()), GetNextFormatSection(), Context.AsCString());
+			continue;
 		}
 
+		} // switch
+
 		if (Context.RawCode() == (ECborCode::Prim | ECborCode::Value_4Bytes))
 		{
 			AddValueToName(Data, MaxLength, GetNextFormatSection(), Context.AsFloat());
@@ -437,6 +457,9 @@ static void AppendMetadataToString(FString& Str, const TraceServices::FMetadataS
 			continue;
 		}
 
+		GetNextFormatSection();
+		AddValueToName(Data, MaxLength, FString(TEXT("%s")), TEXT("???")); // unknown field type
+
 		if (Context.IsFiniteContainer())
 		{
 			CborReader.SkipContainer(ECborCode::Array);
@@ -444,7 +467,7 @@ static void AppendMetadataToString(FString& Str, const TraceServices::FMetadataS
 	}
 
 	// Append what's left of the format string.
-	Str.Append(Format);
+	InOutStr.Append(Format);
 }
 
 ////////////////////////////////////////////////////////////////////////////////////////////////////
@@ -521,7 +544,7 @@ void AddTimingEventToBuilder(
 
 						if (MetadataSpec)
 						{
-							AppendMetadataToString(EventName, MetadataSpec, Metadata);
+							AppendMetadataToString(EventName, MetadataSpec->Format, Metadata);
 						}
 						else
 						{
@@ -904,10 +927,37 @@ void FThreadTimingTrackImpl::InitTooltip(FTooltipDrawState& InOutTooltip, const
 			[this, &InOutTooltip, &TooltipEvent, &ParentTimingEvent, &RootTimingEvent]
 			(const TraceServices::ITimingProfilerProvider& TimingProfilerProvider, const TraceServices::ITimingProfilerTimerReader& TimerReader)
 			{
-				const TraceServices::FTimingProfilerTimer* Timer = TimerReader.GetTimer(TooltipEvent.GetTimerIndex());
+				// Add the timer name.
+				const uint32 TimerIndex = TooltipEvent.GetTimerIndex();
+				const TraceServices::FTimingProfilerTimer* Timer = TimerReader.GetTimer(TimerIndex);
 				const TCHAR* TimerName = (Timer != nullptr) ? Timer->Name : TEXT("N/A");
 				InOutTooltip.AddTitle(TimerName);
 
+				TArrayView<const uint8> Metadata;
+				const TraceServices::FMetadataSpec* MetadataSpec = nullptr;
+				if (int32(TimerIndex) < 0) // has metadata?
+				{
+					Metadata = TimerReader.GetMetadata(TimerIndex);
+					if (Timer && Timer->HasValidMetadataSpecId())
+					{
+						MetadataSpec = TimingProfilerProvider.GetMetadataSpec(Timer->MetadataSpecId);
+					}
+
+#if 0 // for debugging
+					// Add the full timer name (i.e. with formatting/metadata).
+					FString TimerNameEx = TimerName;
+					if (MetadataSpec)
+					{
+						AppendMetadataToString(TimerNameEx, MetadataSpec->Format, Metadata);
+					}
+					else
+					{
+						AppendMetadataToString(TimerNameEx, Metadata);
+					}
+					InOutTooltip.AddTitle(TimerNameEx, FLinearColor(1.0f, 1.0f, 1.0f, 1.0f));
+#endif
+				}
+
 				const double TooltipEventDuration = TooltipEvent.GetDuration();
 
 				if (TooltipEvent.GetDepth() > 0 &&
@@ -950,13 +1000,8 @@ void FThreadTimingTrackImpl::InitTooltip(FTooltipDrawState& InOutTooltip, const
 
 				InOutTooltip.AddNameValueTextLine(TEXTVIEW("Depth:"), FString::Printf(TEXT("%d"), TooltipEvent.GetDepth()));
 
-				const TraceServices::FMetadataSpec* MetadataSpec = nullptr;
-				if (Timer && Timer->HasValidMetadataSpecId())
-				{
-					MetadataSpec = TimingProfilerProvider.GetMetadataSpec(Timer->MetadataSpecId);
-				}
-				TArrayView<const uint8> Metadata = TimerReader.GetMetadata(TooltipEvent.GetTimerIndex());
-				if (Metadata.Num() > 0)
+				// Add detailed metadata.
+				if (int32(TimerIndex) < 0) // has metadata?
 				{
 					AppendMetadataToTooltip(InOutTooltip, MetadataSpec, Metadata);
 				}
@@ -1136,7 +1181,8 @@ void FThreadTimingTrackImpl::OnClipboardCopyEvent(const ITimingEvent& InSelected
 		const FThreadTrackEvent& TrackEvent = InSelectedEvent.As<FThreadTrackEvent>();
 
 		ReadTimers(
-			[this, &TrackEvent](const TraceServices::ITimingProfilerProvider& TimingProfilerProvider, const TraceServices::ITimingProfilerTimerReader& TimerReader)
+			[this, &TrackEvent]
+			(const TraceServices::ITimingProfilerProvider& TimingProfilerProvider, const TraceServices::ITimingProfilerTimerReader& TimerReader)
 			{
 				const TraceServices::FTimingProfilerTimer* TimerPtr = TimerReader.GetTimer(TrackEvent.GetTimerIndex());
 
@@ -1150,16 +1196,18 @@ void FThreadTimingTrackImpl::OnClipboardCopyEvent(const ITimingEvent& InSelected
 					if (int32(TimerIndex) < 0) // has metadata?
 					{
 						TArrayView<const uint8> Metadata = TimerReader.GetMetadata(TimerIndex);
+
 						const TraceServices::FMetadataSpec* MetadataSpec = nullptr;
 						if (TimerPtr && TimerPtr->HasValidMetadataSpecId())
 						{
 							MetadataSpec = TimingProfilerProvider.GetMetadataSpec(TimerPtr->MetadataSpecId);
 						}
+
 						if (Metadata.Num() > 0)
 						{
 							if (MetadataSpec)
 							{
-								AppendMetadataToString(EventName, MetadataSpec, Metadata);
+								AppendMetadataToString(EventName, MetadataSpec->Format, Metadata);
 							}
 							else
 							{
