[Insights] Fix fence tracks visibility toggle and hide them by default for performance reasons

diff --git a/Engine/Config/BaseUnrealInsightsSettings.ini b/Engine/Config/BaseUnrealInsightsSettings.ini
--- a/Engine/Config/BaseUnrealInsightsSettings.ini
+++ b/Engine/Config/BaseUnrealInsightsSettings.ini
@@ -68,7 +68,7 @@ ShowGpuWorkTracks=True
 ShowGpuWorkOverlays=True
 ShowGpuWorkExtendedLines=True
 ShowGpuFencesRelations=True
-ShowGpuFencesTracks=True
+ShowGpuFencesTracks=False
 ShowGpuFencesExtendedLines=True
 
 ShowAllVerseTracks=True
@@ -166,7 +166,7 @@ ShowGpuWorkTracks=True
 ShowGpuWorkOverlays=True
 ShowGpuWorkExtendedLines=True
 ShowGpuFencesRelations=True
-ShowGpuFencesTracks=True
+ShowGpuFencesTracks=False
 ShowGpuFencesExtendedLines=True
 
 ShowAllVerseTracks=False
diff --git a/Engine/Source/Developer/TraceInsights/Private/Insights/TimingProfiler/ViewModels/ThreadTimingSharedState.h b/Engine/Source/Developer/TraceInsights/Private/Insights/TimingProfiler/ViewModels/ThreadTimingSharedState.h
--- a/Engine/Source/Developer/TraceInsights/Private/Insights/TimingProfiler/ViewModels/ThreadTimingSharedState.h
+++ b/Engine/Source/Developer/TraceInsights/Private/Insights/TimingProfiler/ViewModels/ThreadTimingSharedState.h
@@ -281,7 +281,7 @@ private:
 	bool Command_ShowGpuFencesRelations_CanExecute() { return GpuTracks.Num() > 0; }
 	bool Command_ShowGpuFencesRelations_IsChecked() { return AreGpuFenceRelationsVisible(); }
 
-	void Command_ShowGpuFencesTracks_Execute() { ToggleOptionState(FTimingViewSettings::Option_ShowGpuFencesTracks); }
+	void Command_ShowGpuFencesTracks_Execute() { SetGpuFencesTracksVisibility(!AreGpuFencesTracksVisible()); }
 	bool Command_ShowGpuFencesTracks_CanExecute() { return GpuTracks.Num() > 0; }
 	bool Command_ShowGpuFencesTracks_IsChecked() { return AreGpuFencesTracksVisible(); }
 
