[Insights] * Modules/Symbols: Fixed debug symbols failing to use symbols path defined by environment variable (UE_INSIGHTS_SYMBOL_PATH) when
quotation marks are included in the path.
* Simplified the tooltip with Search Paths in the Modules panel.

diff --git a/Engine/Source/Developer/TraceInsights/Private/Insights/MemoryProfiler/Common/SymbolSearchPathsHelper.h b/Engine/Source/Developer/TraceInsights/Private/Insights/MemoryProfiler/Common/SymbolSearchPathsHelper.h
--- a/Engine/Source/Developer/TraceInsights/Private/Insights/MemoryProfiler/Common/SymbolSearchPathsHelper.h
+++ b/Engine/Source/Developer/TraceInsights/Private/Insights/MemoryProfiler/Common/SymbolSearchPathsHelper.h
@@ -30,7 +30,7 @@ public:
 			return FText::Format(LOCTEXT("SymbolResolutionPaths", "Symbol search paths (in order of precedence): {0}"),
 					 FText::FromStringView(PathsListStr));
 		}
-		return LOCTEXT("SymbolResolutionNoPaths", "No symbol paths has been setup. Use configuration file, UE_INSIGHTS_SYMBOL_PATH environment variable or manually select a file or directory from the Modules view.");
+		return LOCTEXT("SymbolResolutionNoPaths", "No symbol paths has been setup.\nUse configuration file, UE_INSIGHTS_SYMBOL_PATH environment variable or manually select a file or directory from the Modules view.");
 	}
 };
 
diff --git a/Engine/Source/Developer/TraceInsights/Private/Insights/Widgets/SModulesView.cpp b/Engine/Source/Developer/TraceInsights/Private/Insights/Widgets/SModulesView.cpp
--- a/Engine/Source/Developer/TraceInsights/Private/Insights/Widgets/SModulesView.cpp
+++ b/Engine/Source/Developer/TraceInsights/Private/Insights/Widgets/SModulesView.cpp
@@ -826,79 +826,30 @@ void SModulesView::Construct(const FArguments& InArgs)
 				.ToolTip(
 					SNew(SToolTip)
 					[
-						SNew(SHorizontalBox)
+						SNew(SVerticalBox)
 
-						+ SHorizontalBox::Slot()
-						.AutoWidth()
+						+ SVerticalBox::Slot()
+						.AutoHeight()
+						.Padding(4.0f, 4.0f)
 						[
-							SNew(SVerticalBox)
-
-							+ SVerticalBox::Slot()
-							.AutoHeight()
-							.Padding(2.0f)
-							[
-								SNew(SSeparator)
-								.Orientation(Orient_Horizontal)
-							]
-
-							+ SVerticalBox::Slot()
-							.AutoHeight()
-							.Padding(2.0f)
-							[
-								SNew(SHorizontalBox)
-
-								// Directory Information
-								+ SHorizontalBox::Slot()
-								.AutoWidth()
-								.Padding(2.0f)
-								[
-									SNew(STextBlock)
-									.Text(LOCTEXT("TT_SymbolDirectoryWarning", ".EXE/.DLL and .PDB files must all be located in the same folder."))
-									.TextStyle(FInsightsStyle::Get(), TEXT("TreeTable.TooltipBold"))
-								]
-							]
-
-							+ SVerticalBox::Slot()
-							.AutoHeight()
-							.Padding(2.0f)
-							[
-								SNew(SSeparator)
-								.Orientation(Orient_Horizontal)
-							]
-
-							+ SVerticalBox::Slot()
-							.AutoHeight()
-							.Padding(2.0f)
-							[
-								SNew(SHorizontalBox)
-
-								// Search Paths
-								+ SHorizontalBox::Slot()
-								.AutoWidth()
-								.Padding(2.0f)
-								[
-									SNew(STextBlock)
-									.Text(LOCTEXT("TT_SearchPaths", "Search Path(s):"))
-									.TextStyle(FInsightsStyle::Get(), TEXT("TreeTable.TooltipBold"))
-								]
-
-								+ SHorizontalBox::Slot()
-								.AutoWidth()
-								.Padding(2.0f)
-								[
-									SNew(STextBlock)
-									.Text(this, &SModulesView::GetSymbolSearchPaths)
-									.TextStyle(FInsightsStyle::Get(), TEXT("TreeTable.Tooltip"))
-								]
-							]
-
-							+ SVerticalBox::Slot()
-							.AutoHeight()
-							.Padding(2.0f)
-							[
-								SNew(SSeparator)
-								.Orientation(Orient_Horizontal)
-							]
+							SNew(STextBlock)
+							.Text(this, &SModulesView::GetSymbolSearchPaths)
+						]
+
+						+ SVerticalBox::Slot()
+						.AutoHeight()
+						.Padding(4.0f, 4.0f)
+						[
+							SNew(SSeparator)
+							.Orientation(Orient_Horizontal)
+						]
+
+						+ SVerticalBox::Slot()
+						.AutoHeight()
+						.Padding(4.0f, 4.0f)
+						[
+							SNew(STextBlock)
+							.Text(LOCTEXT("TT_SymbolDirectoryWarning", ".EXE/.DLL and .PDB files must all be located in the same folder."))
 						]
 					]
 				)
@@ -1019,26 +970,24 @@ FSlateColor SModulesView::GetStatsTextColor() const
 FText SModulesView::GetSymbolSearchPaths() const
 {
 	const TSharedPtr<const TraceServices::IAnalysisSession> Session = FInsightsManager::Get()->GetSession();
-	if (const TraceServices::IModuleProvider* ModuleProvider = Session->ReadProvider<TraceServices::IModuleProvider>("ModuleProvider"))
+	if (Session.IsValid())
 	{
-		if (ModuleProvider)
+		if (const TraceServices::IModuleProvider* ModuleProvider = Session->ReadProvider<TraceServices::IModuleProvider>("ModuleProvider"))
 		{
 			TStringBuilder<2048> PathsListStr;
-			ModuleProvider->EnumerateSymbolSearchPaths([&](FStringView Path)
-			{
-				FString CleanedPath = FString(Path);
-				CleanedPath.TrimQuotesInline();
-				PathsListStr << (TEXT("%s"), PathsListStr.Len() ? "\n" : "") << TEXT("\u2022 ") << CleanedPath;
-			});
-
+			ModuleProvider->EnumerateSymbolSearchPaths([&](FStringView Path) { PathsListStr << TEXT("\n \u2022 ") << Path; });
 			if (PathsListStr.Len())
 			{
-				return FText::FromStringView(PathsListStr);
+				return FText::Format(LOCTEXT("SymbolResolutionPaths", "Symbol search paths (in order of precedence): {0}"),
+					FText::FromStringView(PathsListStr));
+			}
+			else
+			{
+				return LOCTEXT("SymbolResolutionNoPaths", "No symbol paths has been setup.\nUse configuration file, UE_INSIGHTS_SYMBOL_PATH environment variable or manually select a file or directory from the Modules view.");
 			}
 		}
 	}
-
-	return LOCTEXT("SymbolResolutionNoPaths", " No symbol paths have been setup. \n Use configuration file, UE_INSIGHTS_SYMBOL_PATH environment variable or manually select a file or directory from the Modules view.");
+	return FText::GetEmpty();
 }
 
 ////////////////////////////////////////////////////////////////////////////////////////////////////
diff --git a/Engine/Source/Developer/TraceServices/Private/Common/SymbolHelper.cpp b/Engine/Source/Developer/TraceServices/Private/Common/SymbolHelper.cpp
--- a/Engine/Source/Developer/TraceServices/Private/Common/SymbolHelper.cpp
+++ b/Engine/Source/Developer/TraceServices/Private/Common/SymbolHelper.cpp
@@ -26,11 +26,13 @@ void FSymbolHelper::GetPathsFromEnvVars(const FLogCategoryAlias& LogCategory, TA
 			FString SymbolPathPart, SymbolPathRemainder(EnvVariable);
 			while (SymbolPathRemainder.Split(TEXT(";"), &SymbolPathPart, &SymbolPathRemainder))
 			{
+				SymbolPathPart.TrimQuotesInline();
 				if (!SymbolPathPart.IsEmpty())
 				{
 					OutList.Emplace(SymbolPathPart);
 				}
 			}
+			SymbolPathRemainder.TrimQuotesInline();
 			if (!SymbolPathRemainder.IsEmpty())
 			{
 				OutList.Emplace(SymbolPathRemainder);
