[Insights] TimingView: Fixed "Reset to Default" for the GPU track options.

diff --git a/Engine/Source/Developer/TraceInsights/Private/Insights/TimingProfiler/ViewModels/ThreadTimingSharedState.cpp b/Engine/Source/Developer/TraceInsights/Private/Insights/TimingProfiler/ViewModels/ThreadTimingSharedState.cpp
--- a/Engine/Source/Developer/TraceInsights/Private/Insights/TimingProfiler/ViewModels/ThreadTimingSharedState.cpp
+++ b/Engine/Source/Developer/TraceInsights/Private/Insights/TimingProfiler/ViewModels/ThreadTimingSharedState.cpp
@@ -352,8 +352,7 @@ void FThreadTimingSharedState::GetVisibleTimelineIndexes(TSet<uint32>& OutSet) c
 
 void FThreadTimingSharedState::ResetToDefaultValues()
 {
-	RemoveGpuWorkChildTracks();
-	RemoveGpuFencesChildTracks();
+	// Note: The default values are set in RegisterDefaultToggleOptions() from InsightsSettings.h
 	TimingViewSettings->Reset();
 
 	if (TimingView &&
@@ -372,13 +371,9 @@ void FThreadTimingSharedState::ResetToDefaultValues()
 		HideAllCpuSamplingTracks();
 	}
 
-	TimingViewSettings->SetToggleOption(FTimingViewSettings::Option_ShowGpuWorkOverlays, true);
-	TimingViewSettings->SetToggleOption(FTimingViewSettings::Option_ShowGpuWorkExtendedLines, true);
-	SetGpuWorkTracksVisibility(true);
-
-	TimingViewSettings->SetToggleOption(FTimingViewSettings::Option_ShowGpuFencesRelations, true);
-	TimingViewSettings->SetToggleOption(FTimingViewSettings::Option_ShowGpuFencesExtendedLines, true);
-	SetGpuFencesTracksVisibility(true);
+	// Update (create or remove) the Work/Fences child tracks based on the default visibility.
+	SetGpuWorkTracksVisibility(AreGpuWorkTracksVisible());
+	SetGpuFencesTracksVisibility(AreGpuFencesTracksVisible());
 }
 
 ////////////////////////////////////////////////////////////////////////////////////////////////////
@@ -1173,18 +1168,15 @@ TSharedPtr<const ITimingEvent> FThreadTimingSharedState::FindMinEventInstance(ui
 
 void FThreadTimingSharedState::SetGpuWorkTracksVisibility(bool bOnOff)
 {
-	if (TimingViewSettings->IsToggleOptionEnabled(FTimingViewSettings::Option_ShowGpuWorkTracks) != bOnOff)
-	{
-		TimingViewSettings->SetToggleOption(FTimingViewSettings::Option_ShowGpuWorkTracks, bOnOff);
+	TimingViewSettings->SetToggleOption(FTimingViewSettings::Option_ShowGpuWorkTracks, bOnOff);
 
-		if (bOnOff)
-		{
-			AddGpuWorkChildTracks();
-		}
-		else
-		{
-			RemoveGpuWorkChildTracks();
-		}
+	if (bOnOff)
+	{
+		AddGpuWorkChildTracks();
+	}
+	else
+	{
+		RemoveGpuWorkChildTracks();
 	}
 }
 
@@ -1206,15 +1198,21 @@ void FThreadTimingSharedState::AddGpuWorkChildTracks()
 	{
 		// Check if there is an available GPU track for this queue.
 		TSharedPtr<FGpuQueueTimingTrack>* TrackPtrPtr = GpuTracks.Find(QueueInfo.Id);
-
 		if (TrackPtrPtr != nullptr)
 		{
-			// Create the GPU Work track and attach it to the GPU queue track.
-			const FString WorkTrackName = FString::Printf(TEXT("GPU%u - %s %u - WORK"), QueueInfo.GPU, QueueInfo.Name, QueueInfo.Index);
-			TSharedRef<FGpuQueueWorkTimingTrack> WorkTrack = MakeShared<FGpuQueueWorkTimingTrack>(*this, WorkTrackName, QueueInfo.WorkTimelineIndex, QueueInfo.Id);
-			WorkTrack->SetLocation((*TrackPtrPtr)->GetLocation());
-			(*TrackPtrPtr)->AddChildTrack(WorkTrack, 0);
-			WorkTrack->SetParentTrack(*TrackPtrPtr);
+			const TSharedPtr<FGpuQueueTimingTrack>& QueueTrack = *TrackPtrPtr;
+			check(QueueTrack.IsValid());
+
+			// Check if the Work child track is not already created.
+			if (!QueueTrack->FindChildTrackOfType<FGpuQueueWorkTimingTrack>())
+			{
+				// Create the GPU Work track and attach it to the GPU queue track.
+				const FString WorkTrackName = FString::Printf(TEXT("GPU%u - %s %u - Work"), QueueInfo.GPU, QueueInfo.Name, QueueInfo.Index);
+				TSharedRef<FGpuQueueWorkTimingTrack> WorkTrack = MakeShared<FGpuQueueWorkTimingTrack>(*this, WorkTrackName, QueueInfo.WorkTimelineIndex, QueueInfo.Id);
+				WorkTrack->SetLocation(QueueTrack->GetLocation());
+				QueueTrack->AddChildTrack(WorkTrack, 0);
+				WorkTrack->SetParentTrack(QueueTrack);
+			}
 		}
 	});
 }
@@ -1225,16 +1223,32 @@ void FThreadTimingSharedState::RemoveGpuWorkChildTracks()
 {
 	for (auto& Pair : GpuTracks)
 	{
-		const TSharedPtr<FGpuQueueTimingTrack>& Track = Pair.Value;
-		if (TSharedPtr<FGpuQueueWorkTimingTrack> WorkTrack = Track->FindChildTrackOfType<FGpuQueueWorkTimingTrack>())
+		const TSharedPtr<FGpuQueueTimingTrack>& QueueTrack = Pair.Value;
+		check(QueueTrack.IsValid());
+		if (TSharedPtr<FGpuQueueWorkTimingTrack> WorkTrack = QueueTrack->FindChildTrackOfType<FGpuQueueWorkTimingTrack>())
 		{
-			Track->RemoveChildTrack(WorkTrack.ToSharedRef());
+			QueueTrack->RemoveChildTrack(WorkTrack.ToSharedRef());
 		}
 	}
 }
 
 ////////////////////////////////////////////////////////////////////////////////////////////////////
 
+void FThreadTimingSharedState::SetGpuFencesTracksVisibility(bool bOnOff)
+{
+	TimingViewSettings->SetToggleOption(FTimingViewSettings::Option_ShowGpuFencesTracks, bOnOff);
+
+	if (bOnOff)
+	{
+		AddGpuFencesChildTracks();
+	}
+	else
+	{
+		RemoveGpuFencesChildTracks();
+	}
+}
+
+////////////////////////////////////////////////////////////////////////////////////////////////////
 
 void FThreadTimingSharedState::AddGpuFencesChildTracks()
 {
@@ -1252,15 +1266,21 @@ void FThreadTimingSharedState::AddGpuFencesChildTracks()
 	{
 		// Check if there is an available GPU track for this queue.
 		TSharedPtr<FGpuQueueTimingTrack>* TrackPtrPtr = GpuTracks.Find(QueueInfo.Id);
-
 		if (TrackPtrPtr != nullptr)
 		{
-			// Create the GPU Work track and attach it to the GPU queue track.
-			const FString TrackName = FString::Printf(TEXT("GPU%u - %s %u - Fences"), QueueInfo.GPU, QueueInfo.Name, QueueInfo.Index);
-			TSharedRef<FGpuFencesTimingTrack> Track = MakeShared<FGpuFencesTimingTrack>(*this, TrackName, QueueInfo.Id);
-			Track->SetLocation((*TrackPtrPtr)->GetLocation());
-			(*TrackPtrPtr)->AddChildTrack(Track, (*TrackPtrPtr)->GetChildTracks().Num());
-			Track->SetParentTrack(*TrackPtrPtr);
+			const TSharedPtr<FGpuQueueTimingTrack>& QueueTrack = *TrackPtrPtr;
+			check(QueueTrack.IsValid());
+
+			// Check if the Fences child track is not already created.
+			if (!QueueTrack->FindChildTrackOfType<FGpuFencesTimingTrack>())
+			{
+				// Create the GPU Fences track and attach it to the GPU queue track.
+				const FString FencesTrackName = FString::Printf(TEXT("GPU%u - %s %u - Fences"), QueueInfo.GPU, QueueInfo.Name, QueueInfo.Index);
+				TSharedRef<FGpuFencesTimingTrack> FencesTrack = MakeShared<FGpuFencesTimingTrack>(*this, FencesTrackName, QueueInfo.Id);
+				FencesTrack->SetLocation(QueueTrack->GetLocation());
+				QueueTrack->AddChildTrack(FencesTrack, QueueTrack->GetChildTracks().Num());
+				FencesTrack->SetParentTrack(QueueTrack);
+			}
 		}
 	});
 }
@@ -1271,29 +1291,11 @@ void FThreadTimingSharedState::RemoveGpuFencesChildTracks()
 {
 	for (auto& Pair : GpuTracks)
 	{
-		const TSharedPtr<FGpuQueueTimingTrack>& Track = Pair.Value;
-		if (TSharedPtr<FGpuFencesTimingTrack> FencesTrack = Track->FindChildTrackOfType<FGpuFencesTimingTrack>())
-		{
-			Track->RemoveChildTrack(FencesTrack.ToSharedRef());
-		}
-	}
-}
-
-////////////////////////////////////////////////////////////////////////////////////////////////////
-
-void FThreadTimingSharedState::SetGpuFencesTracksVisibility(bool bOnOff)
-{
-	if (TimingViewSettings->IsToggleOptionEnabled(FTimingViewSettings::Option_ShowGpuFencesTracks) != bOnOff)
-	{
-		TimingViewSettings->SetToggleOption(FTimingViewSettings::Option_ShowGpuFencesTracks, bOnOff);
-
-		if (bOnOff)
-		{
-			AddGpuFencesChildTracks();
-		}
-		else
+		const TSharedPtr<FGpuQueueTimingTrack>& QueueTrack = Pair.Value;
+		check(QueueTrack.IsValid());
+		if (TSharedPtr<FGpuFencesTimingTrack> FencesTrack = QueueTrack->FindChildTrackOfType<FGpuFencesTimingTrack>())
 		{
-			RemoveGpuFencesChildTracks();
+			QueueTrack->RemoveChildTrack(FencesTrack.ToSharedRef());
 		}
 	}
 }
