[Insights] * TimingView/Tasks: Fixed crash when double clicking on a task that has undefined timestamps. The code would try to zoom the timing view to fit the respective task (created to finished time range). If one of the timestamps is invalid zooming will result in a crash (in various UI code).
* TimingView/Tasks: Reduced memory usage for tasks analysis and UI (by improving the memory layout for FTaskInfo struct and for the FTaskEntry class).

diff --git a/Engine/Source/Developer/TraceInsights/Private/Insights/TaskGraphProfiler/ViewModels/TaskEntry.h b/Engine/Source/Developer/TraceInsights/Private/Insights/TaskGraphProfiler/ViewModels/TaskEntry.h
--- a/Engine/Source/Developer/TraceInsights/Private/Insights/TaskGraphProfiler/ViewModels/TaskEntry.h
+++ b/Engine/Source/Developer/TraceInsights/Private/Insights/TaskGraphProfiler/ViewModels/TaskEntry.h
@@ -15,13 +15,15 @@ namespace UE::Insights::TaskGraphProfiler
 
 class FTaskEntry
 {
-	friend class STaskTableTreeView;
-
 public:
-	FTaskEntry(const TraceServices::FTaskInfo &TaskInfo);
+	FTaskEntry(const TraceServices::FTaskInfo& TaskInfo);
 	~FTaskEntry() {}
 
+	TaskTrace::FId GetId() const { return Id; }
+
 	const TCHAR* GetDebugName() const { return DebugName; }
+	bool IsTracked() const { return bTracked; }
+	int32 GetThreadToExecuteOn() const { return ThreadToExecuteOn; }
 
 	double GetCreatedTimestamp() const { return CreatedTimestamp; }
 	uint32 GetCreatedThreadId() const { return CreatedThreadId; }
@@ -45,44 +47,37 @@ public:
 
 	uint64 GetTaskSize() const { return TaskSize; }
 
-	uint32 GetNumParents() const { return NumParents; }
-	uint32 GetNumNested() const { return NumNested; }
 	uint32 GetNumPrerequisites() const { return NumPrerequisites; }
 	uint32 GetNumSubsequents() const { return NumSubsequents; }
-
-	TaskTrace::FId GetId() const { return Id; }
+	uint32 GetNumParents() const { return NumParents; }
+	uint32 GetNumNested() const { return NumNested; }
 
 private:
 	TaskTrace::FId Id;
 
-	const TCHAR* DebugName; 
-	bool bTracked; 
+	const TCHAR* DebugName;
+	bool bTracked;
 	int32 ThreadToExecuteOn;
 
 	double CreatedTimestamp;
-	uint32 CreatedThreadId;
-
 	double LaunchedTimestamp;
-	uint32 LaunchedThreadId;
-
 	double ScheduledTimestamp;
-	uint32 ScheduledThreadId;
-
 	double StartedTimestamp;
-	uint32 StartedThreadId;
-
 	double FinishedTimestamp;
-
 	double CompletedTimestamp;
-	uint32 CompletedThreadId;
-
 	double DestroyedTimestamp;
+
+	uint32 CreatedThreadId;
+	uint32 LaunchedThreadId;
+	uint32 ScheduledThreadId;
+	uint32 StartedThreadId;
+	uint32 CompletedThreadId;
 	uint32 DestroyedThreadId;
 
 	uint64 TaskSize;
 
-	uint32 NumSubsequents;
 	uint32 NumPrerequisites;
+	uint32 NumSubsequents;
 	uint32 NumParents;
 	uint32 NumNested;
 };
diff --git a/Engine/Source/Developer/TraceInsights/Private/Insights/TaskGraphProfiler/Widgets/STaskTableTreeView.cpp b/Engine/Source/Developer/TraceInsights/Private/Insights/TaskGraphProfiler/Widgets/STaskTableTreeView.cpp
--- a/Engine/Source/Developer/TraceInsights/Private/Insights/TaskGraphProfiler/Widgets/STaskTableTreeView.cpp
+++ b/Engine/Source/Developer/TraceInsights/Private/Insights/TaskGraphProfiler/Widgets/STaskTableTreeView.cpp
@@ -738,51 +738,37 @@ void STaskTableTreeView::ContextMenu_GoToTask_Execute()
 		return;
 	}
 
-	double Duration = (TaskEntry->GetFinishedTimestamp() - TaskEntry->GetCreatedTimestamp()) * 1.5;
-	TimingView->ZoomOnTimeInterval(TaskEntry->GetCreatedTimestamp() - Duration * 0.15, Duration);
-
-	TSharedPtr<FTaskTimingSharedState> TaskSharedState = FTaskGraphProfilerManager::Get()->GetTaskTimingSharedState();
-
-	if (TaskSharedState.IsValid() && FTaskGraphProfilerManager::Get()->GetShowAnyRelations())
+	double CreatedTimestamp = TaskEntry->GetCreatedTimestamp();
+	if (CreatedTimestamp == TraceServices::FTaskInfo::InvalidTimestamp)
 	{
-		TaskSharedState->SetTaskId(TaskEntry->GetId());
+		CreatedTimestamp = TimingView->GetViewport().GetMinValidTime();
 	}
-
-	TSharedPtr<TimingProfiler::FThreadTimingSharedState> ThreadTimingState = TimingView->GetThreadTimingSharedState();
-	if (!ThreadTimingState.IsValid())
+	double FinishedTimestamp = TaskEntry->GetFinishedTimestamp();
+	if (FinishedTimestamp == TraceServices::FTaskInfo::InvalidTimestamp)
 	{
-		return;
+		FinishedTimestamp = TimingView->GetViewport().GetMaxValidTime();
 	}
+	const double Duration = FinishedTimestamp - CreatedTimestamp;
+	TimingView->ZoomOnTimeInterval(CreatedTimestamp - Duration * 0.15, Duration * 1.30);
 
-	TSharedPtr<TimingProfiler::FCpuTimingTrack> Track = ThreadTimingState->GetCpuTrack(TaskEntry->GetStartedThreadId());
-	if (!Track.IsValid())
+	if (TSharedPtr<FTaskTimingSharedState> TaskSharedState = FTaskGraphProfilerManager::Get()->GetTaskTimingSharedState())
 	{
-		return;
+		if (FTaskGraphProfilerManager::Get()->GetShowAnyRelations())
+		{
+			TaskSharedState->SetTaskId(TaskEntry->GetId());
+		}
 	}
 
-	TimingView->SelectTimingTrack(Track, true);
-
-	auto SearchFilter = [TaskEntry](double StartTime, double EndTime, uint32 Depth)
+	if (TSharedPtr<FThreadTimingSharedState> ThreadTimingState = TimingView->GetThreadTimingSharedState())
 	{
-		if (StartTime >= TaskEntry->GetStartedTimestamp() && EndTime <= TaskEntry->GetFinishedTimestamp())
+		if (TSharedPtr<FCpuTimingTrack> Track = ThreadTimingState->GetCpuTrack(TaskEntry->GetStartedThreadId()))
 		{
-			return true;
+			TimingView->SelectTimingTrack(nullptr, false); // to force bring the track into view, if already selected
+			TimingView->SelectTimingTrack(Track, true);
 		}
-
-		return false;
-	};
-
-	FTimingEventSearchParameters SearchParams(TaskEntry->StartedTimestamp, TaskEntry->FinishedTimestamp, ETimingEventSearchFlags::StopAtFirstMatch, SearchFilter);
-
-	const TSharedPtr<const ITimingEvent> FoundEvent = Track->SearchEvent(SearchParams);
-	if (FoundEvent.IsValid() && Track->IsVisible())
-	{
-		TimingView->SelectTimingEvent(FoundEvent, true);
-	}
-	else
-	{
-		FTaskGraphProfilerManager::Get()->ShowTaskRelations(TaskEntry->GetId());
 	}
+
+	FTaskGraphProfilerManager::Get()->ShowTaskRelations(TaskEntry->GetId());
 }
 
 ////////////////////////////////////////////////////////////////////////////////////////////////////
@@ -823,12 +809,12 @@ bool STaskTableTreeView::GetSourceFileAndLineForSelectedTask(FString& OutFile, u
 	FTaskNodePtr SelectedTask = StaticCastSharedPtr<FTaskNode>(SelectedTreeNode);
 
 	const FTaskEntry* TaskEntry = SelectedTask->GetTask();
-	if (TaskEntry == nullptr || TaskEntry->DebugName == nullptr)
+	if (TaskEntry == nullptr || TaskEntry->GetDebugName() == nullptr)
 	{
 		return false;
 	}
 
-	const FString DebugName = TaskEntry->DebugName;
+	const FString DebugName = TaskEntry->GetDebugName();
 	if (DebugName.Len() < 4 ||
 		DebugName[DebugName.Len() - 1] != TEXT(')'))
 	{
@@ -911,7 +897,7 @@ void STaskTableTreeView::SearchForItem(TSharedPtr<FTableTaskCancellationToken> C
 			break;
 		}
 
-		if (Tasks[Index].Id == TaskIdToSelect)
+		if (Tasks[Index].GetId() == TaskIdToSelect)
 		{
 			TGraphTask<FSelectNodeByTableRowIndexTask>::CreateTask().ConstructAndDispatchWhenReady(CancellationToken, SharedThis(this), Index);
 			break;
diff --git a/Engine/Source/Developer/TraceServices/Public/TraceServices/Model/TasksProfiler.h b/Engine/Source/Developer/TraceServices/Public/TraceServices/Model/TasksProfiler.h
--- a/Engine/Source/Developer/TraceServices/Public/TraceServices/Model/TasksProfiler.h
+++ b/Engine/Source/Developer/TraceServices/Public/TraceServices/Model/TasksProfiler.h
@@ -31,24 +31,26 @@ struct FTaskInfo
 
 	// some tasks can be created with delayed launch, otherwise "created" and "launched" timestamps and thread IDs are equal
 	double CreatedTimestamp = InvalidTimestamp;
-	uint32 CreatedThreadId;
 	// the moment and the place when the task was launched
 	double LaunchedTimestamp = InvalidTimestamp;
-	uint32 LaunchedThreadId;
 	// the moment and the place when all task's dependencies are resolved and it's added to the execution queue
 	double ScheduledTimestamp = InvalidTimestamp;
-	uint32 ScheduledThreadId;
 	// execution started
 	double StartedTimestamp = InvalidTimestamp;
-	uint32 StartedThreadId;
 	// execution finished, but the task can be still not completed waiting for all nested tasks to complete
 	double FinishedTimestamp = InvalidTimestamp;
 	// execution done and all nested tasks are completed
 	double CompletedTimestamp = InvalidTimestamp;
-	uint32 CompletedThreadId;
 	// the last reference is released and the task is destroyed
 	double DestroyedTimestamp = InvalidTimestamp;
-	uint32 DestroyedThreadId;
+
+	uint32 CreatedThreadId = 0;
+	uint32 LaunchedThreadId = 0;
+	uint32 ScheduledThreadId = 0;
+	uint32 StartedThreadId = 0;
+	uint32 CompletedThreadId = 0;
+	uint32 DestroyedThreadId = 0;
+
 	// task size including user-provided task body
 	uint64 TaskSize = 0;
 
