[Insights] Fix inconsistent path names for editor is built with VFS

diff --git a/Engine/Source/Developer/TraceInsights/Private/Insights/MemoryProfiler/ViewModels/CallstackFormatting.cpp b/Engine/Source/Developer/TraceInsights/Private/Insights/MemoryProfiler/ViewModels/CallstackFormatting.cpp
--- a/Engine/Source/Developer/TraceInsights/Private/Insights/MemoryProfiler/ViewModels/CallstackFormatting.cpp
+++ b/Engine/Source/Developer/TraceInsights/Private/Insights/MemoryProfiler/ViewModels/CallstackFormatting.cpp
@@ -6,6 +6,9 @@
 #include "TraceServices/Model/Callstack.h"
 #include "TraceServices/Model/Modules.h"
 
+// TraceInsights
+#include "Insights/InsightsManager.h"
+
 #define LOCTEXT_NAMESPACE "UE::Insights::CallstackFormatting"
 
 namespace UE::Insights
@@ -61,7 +64,9 @@ void FormatStackFrame(const TraceServices::FStackFrame& Frame, FStringBuilderBas
 			}
 			if (EnumHasAnyFlags(FormatFlags, EStackFrameFormatFlags::File))
 			{
-				OutString.Append(Frame.Symbol->File);
+				FString UsablePath;
+				FInsightsManager::Get()->GetSourceFilePathHelper().GetUsableFilePath(Frame.Symbol->File, UsablePath);
+				OutString.Append(UsablePath);
 			}
 			if (EnumHasAnyFlags(FormatFlags, EStackFrameFormatFlags::Line))
 			{
diff --git a/Engine/Source/Developer/TraceInsights/Private/Insights/ViewModels/LogMessage.cpp b/Engine/Source/Developer/TraceInsights/Private/Insights/ViewModels/LogMessage.cpp
--- a/Engine/Source/Developer/TraceInsights/Private/Insights/ViewModels/LogMessage.cpp
+++ b/Engine/Source/Developer/TraceInsights/Private/Insights/ViewModels/LogMessage.cpp
@@ -12,6 +12,7 @@
 #include "InsightsCore/Common/TimeUtils.h"
 
 // TraceInsights
+#include "Insights/InsightsManager.h"
 #include "Insights/Log.h"
 
 #define LOCTEXT_NAMESPACE "UE::Insights::SLogView"
@@ -24,7 +25,6 @@ FLogMessageRecord::FLogMessageRecord(const TraceServices::FLogMessageInfo& Trace
 	: Time(TraceLogMessage.Time)
 	//, Category(FText::FromString(TraceLogMessage.Category))
 	//, Message(FText::FromString(TraceLogMessage.Message))
-	, File(FText::FromString(TraceLogMessage.File))
 	, Line(TraceLogMessage.Line)
 	, Index(static_cast<uint32>(TraceLogMessage.Index))
 	, Verbosity(TraceLogMessage.Verbosity)
@@ -37,6 +37,10 @@ FLogMessageRecord::FLogMessageRecord(const TraceServices::FLogMessageInfo& Trace
 	}
 	Category = FText::FromString(MoveTemp(CategoryStr));
 
+	FString UsableFile;
+	UE::Insights::FInsightsManager::Get()->GetSourceFilePathHelper().GetUsableFilePath(TraceLogMessage.File, UsableFile);
+	File = FText::FromString(UsableFile);
+
 	// Strip the trailing whitespaces (ex. some messages ends with "\n" and we do not want the LogView rows to have an unnecessary increased height).
 	FString MessageStr(TraceLogMessage.Message);
 	MessageStr.TrimEndInline();
diff --git a/Engine/Source/Developer/TraceInsights/Private/Insights/Widgets/SLogView.cpp b/Engine/Source/Developer/TraceInsights/Private/Insights/Widgets/SLogView.cpp
--- a/Engine/Source/Developer/TraceInsights/Private/Insights/Widgets/SLogView.cpp
+++ b/Engine/Source/Developer/TraceInsights/Private/Insights/Widgets/SLogView.cpp
@@ -2113,18 +2113,15 @@ void SLogView::OpenSource() const
 		const FString File = Record.GetFile();
 		const uint32 Line = Record.GetLine();
 		
-		FString UsablePath;
-		const bool bPathExists = FInsightsManager::Get()->GetSourceFilePathHelper().GetUsableFilePath(File, UsablePath);
-
 		ISourceCodeAccessModule& SourceCodeAccessModule = FModuleManager::LoadModuleChecked<ISourceCodeAccessModule>("SourceCodeAccess");
-		if (bPathExists)
+		if (FPaths::FileExists(File))
 		{
 			ISourceCodeAccessor& SourceCodeAccessor = SourceCodeAccessModule.GetAccessor();
-			SourceCodeAccessor.OpenFileAtLine(UsablePath, Line);
+			SourceCodeAccessor.OpenFileAtLine(File, Line);
 		}
 		else
 		{
-			SourceCodeAccessModule.OnOpenFileFailed().Broadcast(UsablePath);
+			SourceCodeAccessModule.OnOpenFileFailed().Broadcast(File);
 		}
 	}
 }
