[Insights] Fix GPU fences relations drawing issue and disable GPU fence visibilty by default.

diff --git a/Engine/Source/Developer/TraceInsights/Private/Insights/TimingProfiler/ViewModels/GpuFenceRelation.cpp b/Engine/Source/Developer/TraceInsights/Private/Insights/TimingProfiler/ViewModels/GpuFenceRelation.cpp
--- a/Engine/Source/Developer/TraceInsights/Private/Insights/TimingProfiler/ViewModels/GpuFenceRelation.cpp
+++ b/Engine/Source/Developer/TraceInsights/Private/Insights/TimingProfiler/ViewModels/GpuFenceRelation.cpp
@@ -83,8 +83,16 @@ void FGpuFenceRelation::Draw(const FDrawContext& DrawContext, const FTimingTrack
 	{
 		int32 ActualSourceDepth = FMath::Min(SourceDepth, MaxEventDepth);
 		Y1 = SourceTrackShared->GetPosY();
-		Y1 += Viewport.GetLayout().GetLaneY(ActualSourceDepth) + Viewport.GetLayout().EventH / 2.0f;
-		Y1 += SourceTrackShared->GetChildTracksTopHeight(Viewport.GetLayout());
+		float ChildTracksHeight = SourceTrackShared->GetChildTracksTopHeight(Viewport.GetLayout());
+		if (ActualSourceDepth >= 0)
+		{
+			Y1 += Viewport.GetLayout().GetLaneY(ActualSourceDepth) + Viewport.GetLayout().EventH / 2.0f;
+		}
+		else if(ChildTracksHeight > 0.0f)
+		{
+			Y1 -= Viewport.GetLayout().EventH / 2.0f;
+		}
+		Y1 += ChildTracksHeight;
 	}
 	else
 	{
@@ -95,8 +103,16 @@ void FGpuFenceRelation::Draw(const FDrawContext& DrawContext, const FTimingTrack
 	{
 		int32 ActualTargetDepth = FMath::Min(TargetDepth, MaxEventDepth);
 		Y2 = TargetTrackShared->GetPosY();
-		Y2 += Viewport.GetLayout().GetLaneY(ActualTargetDepth) + Viewport.GetLayout().EventH / 2.0f;
-		Y2 += TargetTrackShared->GetChildTracksTopHeight(Viewport.GetLayout());
+		float ChildTracksHeight = TargetTrackShared->GetChildTracksTopHeight(Viewport.GetLayout());
+		if (ActualTargetDepth >= 0)
+		{
+			Y2 += Viewport.GetLayout().GetLaneY(ActualTargetDepth) + Viewport.GetLayout().EventH / 2.0f;
+		}
+		else if (ChildTracksHeight > 0.0f)
+		{
+			Y2 -= Viewport.GetLayout().EventH / 2.0f;
+		}
+		Y2 += ChildTracksHeight;
 	}
 	else
 	{
diff --git a/Engine/Source/Developer/TraceInsights/Private/Insights/TimingProfiler/ViewModels/ThreadTimingSharedState.cpp b/Engine/Source/Developer/TraceInsights/Private/Insights/TimingProfiler/ViewModels/ThreadTimingSharedState.cpp
--- a/Engine/Source/Developer/TraceInsights/Private/Insights/TimingProfiler/ViewModels/ThreadTimingSharedState.cpp
+++ b/Engine/Source/Developer/TraceInsights/Private/Insights/TimingProfiler/ViewModels/ThreadTimingSharedState.cpp
@@ -1384,7 +1384,7 @@ void FThreadTimingSharedState::OnTimingEventSelected(TSharedPtr<const ITimingEve
 		int32 SourceDepth = 0;
 		TimingProfilerProvider->ReadTimeline(Index, [SignalFence, &SourceDepth](const TraceServices::ITimingProfilerProvider::Timeline& Timeline)
 		{
-			SourceDepth = FMath::Max(Timeline.GetDepthAt(SignalFence.Timestamp) - 1, 0);
+			SourceDepth = Timeline.GetDepthAt(SignalFence.Timestamp) - 1;
 		});
 
 		Result = TimingProfilerProvider->GetGpuQueueTimelineIndex(WaitFenceQueueId, Index);
@@ -1396,7 +1396,7 @@ void FThreadTimingSharedState::OnTimingEventSelected(TSharedPtr<const ITimingEve
 		int32 TargetDepth = 0;
 		TimingProfilerProvider->ReadTimeline(Index, [WaitFence, &TargetDepth](const TraceServices::ITimingProfilerProvider::Timeline& Timeline)
 		{
-			TargetDepth = FMath::Max(Timeline.GetDepthAt(WaitFence.Timestamp) - 1, 0);
+			TargetDepth = Timeline.GetDepthAt(WaitFence.Timestamp) - 1;
 		});
 
 		TUniquePtr<ITimingEventRelation> RelationBase = MakeUnique<FGpuFenceRelation>(SignalFence.Timestamp, SignalFenceQueueId, WaitFence.Timestamp, WaitFenceQueueId);
