[Insights] Memory Tags: Fixed sorting by name for the Hierarchy column.

diff --git a/Engine/Source/Developer/TraceInsights/Private/Insights/MemoryProfiler/ViewModels/MemTagTable.cpp b/Engine/Source/Developer/TraceInsights/Private/Insights/MemoryProfiler/ViewModels/MemTagTable.cpp
--- a/Engine/Source/Developer/TraceInsights/Private/Insights/MemoryProfiler/ViewModels/MemTagTable.cpp
+++ b/Engine/Source/Developer/TraceInsights/Private/Insights/MemoryProfiler/ViewModels/MemTagTable.cpp
@@ -179,7 +179,7 @@ void FMemTagTable::AddDefaultColumns()
 	int32 ColumnIndex = 0;
 
 	//////////////////////////////////////////////////
-	// Tag Id Column
+	// Tag Name Column
 	{
 		TSharedRef<FTableColumn> ColumnRef = MakeShared<FTableColumn>(FMemTagTableColumns::TagNameColumnId);
 		FTableColumn& Column = *ColumnRef;
@@ -223,6 +223,10 @@ void FMemTagTable::AddDefaultColumns()
 		Column.SetAggregation(ETableColumnAggregation::SameValue);
 
 		AddColumn(ColumnRef);
+
+		// Set the same "by name" sorter also for the hierarchy column.
+		FTableColumn& HierarchyColumn = *GetColumns()[0];
+		HierarchyColumn.SetValueSorter(Sorter);
 	}
 	//////////////////////////////////////////////////
 	// Tag Id Column
