Fix of System.Insights.Analysis.TimingInsights SumTimerIndex assertion fail

#rnx

diff --git a/Engine/Source/Developer/TraceInsights/Private/Insights/Tests/FunctionalTests/TimingProfilerTests.cpp b/Engine/Source/Developer/TraceInsights/Private/Insights/Tests/FunctionalTests/TimingProfilerTests.cpp
--- a/Engine/Source/Developer/TraceInsights/Private/Insights/Tests/FunctionalTests/TimingProfilerTests.cpp
+++ b/Engine/Source/Developer/TraceInsights/Private/Insights/Tests/FunctionalTests/TimingProfilerTests.cpp
@@ -49,7 +49,7 @@ bool FEnumerateTest::RunTest(const FString& Parameters)
 	TestEqual(TEXT("TotalEventDuration"), CheckValues.TotalEventDuration, 680.943945, 1.e-6);
 	TestEqual(TEXT("EventCount"), CheckValues.EventCount, 10836057ull);
 	TestEqual(TEXT("SumDepth"), (uint64)CheckValues.SumDepth, 80030008ull);
-	TestEqual(TEXT("SumTimerIndex"), (uint64)CheckValues.SumTimerIndex, 4126772211ull);
+	TestEqual(TEXT("SumTimerIndex"), (uint64)CheckValues.SumTimerIndex, 4137608268ull);
 
 	AddInfo(FString::Printf(TEXT("Enumeration Duration: %f seconds."), CheckValues.EnumerationDuration));
 
@@ -75,7 +75,7 @@ bool FEnumeratePerformanceTest::RunTest(const FString& Parameters)
 	TestEqual(TEXT("TotalEventDuration"), CheckValues.TotalEventDuration, 10912.775537, 1.e-6);
 	TestEqual(TEXT("EventCount"), CheckValues.EventCount, 137000700ull);
 	TestEqual(TEXT("SumDepth"), (uint64)CheckValues.SumDepth, 1134384338ull);
-	TestEqual(TEXT("SumTimerIndex"), (uint64)CheckValues.SumTimerIndex, 3499618755ull);
+	TestEqual(TEXT("SumTimerIndex"), (uint64)CheckValues.SumTimerIndex, 3636619455ull);
 
 	const double BenchmarkBaseline = 16.0;
 	AddInfo(FString::Printf(TEXT("Enumeration Duration: %f seconds."), CheckValues.EnumerationDuration));
@@ -111,7 +111,7 @@ bool FEnumerateFastTest::RunTest(const FString& Parameters)
 	TestEqual(TEXT("TotalEventDuration"), CheckValues.TotalEventDuration, 1647.693886, 1.e-6);
 	TestEqual(TEXT("EventCount"), CheckValues.EventCount, 1759740ull);
 	TestEqual(TEXT("SumDepth"), (uint64)CheckValues.SumDepth, 15189227ull);
-	TestEqual(TEXT("SumTimerIndex"), (uint64)CheckValues.SumTimerIndex, 1239801518ull);
+	TestEqual(TEXT("SumTimerIndex"), (uint64)CheckValues.SumTimerIndex, 1241561258ull);
 
 	AddInfo(FString::Printf(TEXT("Enumeration Duration: %f seconds."), CheckValues.EnumerationDuration));
 
