[Insights] Fixed use of potential invalid Session pointer.

#rnx

diff --git a/Engine/Source/Developer/TraceInsights/Private/Insights/Widgets/SSessionInfoWindow.cpp b/Engine/Source/Developer/TraceInsights/Private/Insights/Widgets/SSessionInfoWindow.cpp
--- a/Engine/Source/Developer/TraceInsights/Private/Insights/Widgets/SSessionInfoWindow.cpp
+++ b/Engine/Source/Developer/TraceInsights/Private/Insights/Widgets/SSessionInfoWindow.cpp
@@ -359,52 +359,8 @@ void SSessionInfoWindow::Tick(const FGeometry& AllottedGeometry, const double In
 		AnalysisSession = Session;
 		bIsSessionInfoSet = false;
 
-		ETraceStreamType TraceStreamType = FInsightsManager::Get()->GetTraceStreamType();
-
-		switch (TraceStreamType)
-		{
-			case ETraceStreamType::TraceFile:
-				StreamTypeText = LOCTEXT("TraceStreamType_TraceFile", "Trace File");
-				break;
-
-			case ETraceStreamType::TraceStore:
-				StreamTypeText = FText::Format(LOCTEXT("TraceStreamType_TraceStore_Fmt", "Trace Store {0} -- TraceId={1}"),
-					FText::FromString(FString::Printf(TEXT("%s:%u"), *FInsightsManager::Get()->GetLastStoreHost(), FInsightsManager::Get()->GetLastStorePort())),
-					FText::FromString(FString::Printf(TEXT("0x%X"), Session->GetTraceId())));
-				break;
-
-			case ETraceStreamType::DirectTrace:
-				StreamTypeText = FText::Format(LOCTEXT("TraceStreamType_DirectTrace_Fmt", "Direct Trace -- Port={0}"),
-					FText::AsNumber(uint32(FInsightsManager::Get()->GetDirectTracePort()), &FNumberFormattingOptions::DefaultNoGrouping()));
-				break;
-
-			case ETraceStreamType::Custom:
-				StreamTypeText = LOCTEXT("TraceStreamType_Custom", "Custom Stream");
-				break;
-
-			default:
-				StreamTypeText = LOCTEXT("TraceStreamType_Unknown", "Unknown");
-				break;
-		}
-
-		// We can quickly get the session name and uri.
-		{
-			TraceServices::FAnalysisSessionReadScope SessionReadScope(*Session.Get());
-
-			if (TraceStreamType == ETraceStreamType::TraceFile ||
-				TraceStreamType == ETraceStreamType::TraceStore)
-			{
-				SessionNameText = FText::FromString(FPaths::GetBaseFilename(Session->GetName()));
-
-				FString Uri(Session->GetName());
-				FPaths::NormalizeFilename(Uri);
-				UriText = FText::FromString(Uri);
-			}
-			else
-			{
-				SessionNameText = FText::FromString(Session->GetName());
-			}
-		}
+		StreamTypeText = FText::GetEmpty();
+		SessionNameText = FText::GetEmpty();
 
 		PlatformText = FText::GetEmpty();
 		AppNameText = FText::GetEmpty();
@@ -423,131 +379,194 @@ void SSessionInfoWindow::Tick(const FGeometry& AllottedGeometry, const double In
 	// If we already have the session info data, we no longer poll for it.
 	if (!bIsSessionInfoSet && Session.IsValid())
 	{
-		TraceServices::FAnalysisSessionReadScope SessionReadScope(*Session.Get());
+		if (StreamTypeText.IsEmpty())
+		{
+			UpdateStreamTypeAndSessionName(*Session.Get());
+		}
 
-		const TraceServices::IDiagnosticsProvider* DiagnosticsProvider = TraceServices::ReadDiagnosticsProvider(*Session.Get());
+		bIsSessionInfoSet = TryUpdateDiagnosticsInfo(*Session.Get());
+	}
+}
+
+////////////////////////////////////////////////////////////////////////////////////////////////////
+
+void SSessionInfoWindow::UpdateStreamTypeAndSessionName(const TraceServices::IAnalysisSession& Session)
+{
+	ETraceStreamType TraceStreamType = FInsightsManager::Get()->GetTraceStreamType();
+
+	switch (TraceStreamType)
+	{
+	case ETraceStreamType::TraceFile:
+		StreamTypeText = LOCTEXT("TraceStreamType_TraceFile", "Trace File");
+		break;
+
+	case ETraceStreamType::TraceStore:
+		StreamTypeText = FText::Format(LOCTEXT("TraceStreamType_TraceStore_Fmt", "Trace Store {0} -- TraceId={1}"),
+			FText::FromString(FString::Printf(TEXT("%s:%u"), *FInsightsManager::Get()->GetLastStoreHost(), FInsightsManager::Get()->GetLastStorePort())),
+			FText::FromString(FString::Printf(TEXT("0x%X"), Session.GetTraceId())));
+		break;
+
+	case ETraceStreamType::DirectTrace:
+		StreamTypeText = FText::Format(LOCTEXT("TraceStreamType_DirectTrace_Fmt", "Direct Trace -- Port={0}"),
+			FText::AsNumber(uint32(FInsightsManager::Get()->GetDirectTracePort()), &FNumberFormattingOptions::DefaultNoGrouping()));
+		break;
+
+	case ETraceStreamType::Custom:
+		StreamTypeText = LOCTEXT("TraceStreamType_Custom", "Custom Stream");
+		break;
+
+	default:
+		StreamTypeText = LOCTEXT("TraceStreamType_Unknown", "Unknown");
+		break;
+	}
+
+	// We can quickly get the session name and uri.
+	if (TraceStreamType == ETraceStreamType::TraceFile ||
+		TraceStreamType == ETraceStreamType::TraceStore)
+	{
+		SessionNameText = FText::FromString(FPaths::GetBaseFilename(Session.GetName()));
+
+		FString Uri(Session.GetName());
+		FPaths::NormalizeFilename(Uri);
+		UriText = FText::FromString(Uri);
+	}
+	else
+	{
+		SessionNameText = FText::FromString(Session.GetName());
+	}
+}
+
+////////////////////////////////////////////////////////////////////////////////////////////////////
+
+bool SSessionInfoWindow::TryUpdateDiagnosticsInfo(const TraceServices::IAnalysisSession& Session)
+{
+	TraceServices::FAnalysisSessionReadScope SessionReadScope(Session);
 
-		if (DiagnosticsProvider && DiagnosticsProvider->IsSessionInfoAvailable())
+	const TraceServices::IDiagnosticsProvider* DiagnosticsProvider = TraceServices::ReadDiagnosticsProvider(Session);
+	if (!DiagnosticsProvider || !DiagnosticsProvider->IsSessionInfoAvailable())
+	{
+		return false;
+	}
+
+	TraceServices::FSessionInfo SessionInfo = DiagnosticsProvider->GetSessionInfo();
+
+	PlatformText = FText::FromString(SessionInfo.Platform);
+	AppNameText = FText::FromString(SessionInfo.AppName);
+	ProjectNameText = FText::FromString(SessionInfo.ProjectName);
+
+	BuildConfigurationTypeText = FText::FromString(LexToString(SessionInfo.ConfigurationType));
+	BuildTargetTypeText = FText::FromString(LexToString(SessionInfo.TargetType));
+	BuildVersionText = FText::FromString(SessionInfo.BuildVersion);
+	BranchText = FText::FromString(SessionInfo.Branch);
+	ChangelistText = FText::AsNumber(SessionInfo.Changelist, &FNumberFormattingOptions::DefaultNoGrouping());
+
+	EngineVersionText = FText::FromString(SessionInfo.EngineVersion);
+	if (!SessionInfo.EngineVersion.IsEmpty() &&
+		!SessionInfo.BuildVersion.IsEmpty())
+	{
+		// Trim the BuildVersion from the EngineVersion, if it is included and it is the same.
+		int32 Index = SessionInfo.EngineVersion.Find(SessionInfo.BuildVersion, ESearchCase::CaseSensitive);
+		if (Index > 0 && SessionInfo.EngineVersion[Index - 1] == TEXT('-'))
 		{
-			TraceServices::FSessionInfo SessionInfo = DiagnosticsProvider->GetSessionInfo();
+			EngineVersionText = FText::FromString(SessionInfo.EngineVersion.Left(Index - 1));
+		}
+	}
 
-			PlatformText = FText::FromString(SessionInfo.Platform);
-			AppNameText = FText::FromString(SessionInfo.AppName);
-			ProjectNameText = FText::FromString(SessionInfo.ProjectName);
+	CommandLineText = FText::FromString(SessionInfo.CommandLine);
+
+	TArray<FString> VFSPaths;
+	SessionInfo.VFSPaths.ParseIntoArray(VFSPaths, TEXT(";"), true);
+	int32 NumVFSPaths = 0;
+	TStringBuilder<1024> VFSPathsBuilder;
+	for (FString& Path : VFSPaths)
+	{
+		Path.TrimStartAndEndInline();
+		if (!Path.IsEmpty())
+		{
+			if (NumVFSPaths & 1)
+			{
+				VFSPathsBuilder.Append(TEXT(" \u2192 "));
+			}
+			else if (NumVFSPaths > 0)
+			{
+				VFSPathsBuilder.Append(TEXT("\n"));
+			}
+			++NumVFSPaths;
+			VFSPathsBuilder.Append(Path);
+		}
+	}
+	VFSPathsText = FText::FromString(VFSPathsBuilder.ToString());
 
-			BuildConfigurationTypeText = FText::FromString(LexToString(SessionInfo.ConfigurationType));
-			BuildTargetTypeText = FText::FromString(LexToString(SessionInfo.TargetType));
-			BuildVersionText = FText::FromString(SessionInfo.BuildVersion);
-			BranchText = FText::FromString(SessionInfo.Branch);
-			ChangelistText = FText::AsNumber(SessionInfo.Changelist, &FNumberFormattingOptions::DefaultNoGrouping());
+	TStringBuilder<1024> OtherMetadata;
 
-			EngineVersionText = FText::FromString(SessionInfo.EngineVersion);
-			if (!SessionInfo.EngineVersion.IsEmpty() &&
-				!SessionInfo.BuildVersion.IsEmpty())
+	Session.EnumerateMetadata(
+		[&OtherMetadata]
+		(const TraceServices::FTraceSessionMetadata& Metadata)
+		{
+			static FName ExcludedMetadata[]
+			{
+				FName("Platform"),
+				FName("AppName"),
+				FName("ProjectName"),
+				FName("ConfigurationType"),
+				FName("TargetType"),
+				FName("BuildVersion"),
+				FName("Branch"),
+				FName("Changelist"),
+				FName("EngineVersion"),
+				FName("CommandLine"),
+				FName("VFSPaths"),
+			};
+
+			for (int32 Index = 0; Index < UE_ARRAY_COUNT(ExcludedMetadata); ++Index)
 			{
-				// Trim the BuildVersion from the EngineVersion, if it is included and it is the same.
-				int32 Index = SessionInfo.EngineVersion.Find(SessionInfo.BuildVersion, ESearchCase::CaseSensitive);
-				if (Index > 0 && SessionInfo.EngineVersion[Index - 1] == TEXT('-'))
+				if (Metadata.Name == ExcludedMetadata[Index])
 				{
-					EngineVersionText = FText::FromString(SessionInfo.EngineVersion.Left(Index - 1));
+					return;
 				}
 			}
 
-			CommandLineText = FText::FromString(SessionInfo.CommandLine);
-
-			TArray<FString> VFSPaths;
-			SessionInfo.VFSPaths.ParseIntoArray(VFSPaths, TEXT(";"), true);
-			int32 NumVFSPaths = 0;
-			TStringBuilder<1024> VFSPathsBuilder;
-			for (FString& Path : VFSPaths)
+			switch (Metadata.Type)
 			{
-				Path.TrimStartAndEndInline();
-				if (!Path.IsEmpty())
+				case TraceServices::FTraceSessionMetadata::EType::Int64:
 				{
-					if (NumVFSPaths & 1)
+					if (OtherMetadata.Len() > 0)
 					{
-						VFSPathsBuilder.Append(TEXT(" \u2192 "));
+						OtherMetadata.Append(TEXT("\n"));
 					}
-					else if (NumVFSPaths > 0)
-					{
-						VFSPathsBuilder.Append(TEXT("\n"));
-					}
-					++NumVFSPaths;
-					VFSPathsBuilder.Append(Path);
+					OtherMetadata.Append(Metadata.Name.GetPlainNameString());
+					OtherMetadata.Append(TEXT("="));
+					OtherMetadata.Appendf(TEXT("%lld"), Metadata.Int64Value);
+					break;
 				}
-			}
-			VFSPathsText = FText::FromString(VFSPathsBuilder.ToString());
-
-			TStringBuilder<1024> OtherMetadata;
-
-			Session->EnumerateMetadata([&OtherMetadata](const TraceServices::FTraceSessionMetadata& Metadata)
+				case TraceServices::FTraceSessionMetadata::EType::Double:
 				{
-					static FName ExcludedMetadata[]
+					if (OtherMetadata.Len() > 0)
 					{
-						FName("Platform"),
-						FName("AppName"),
-						FName("ProjectName"),
-						FName("ConfigurationType"),
-						FName("TargetType"),
-						FName("BuildVersion"),
-						FName("Branch"),
-						FName("Changelist"),
-						FName("EngineVersion"),
-						FName("CommandLine"),
-						FName("VFSPaths"),
-					};
-
-					for (int32 Index = 0; Index < UE_ARRAY_COUNT(ExcludedMetadata); ++Index)
-					{
-						if (Metadata.Name == ExcludedMetadata[Index])
-						{
-							return;
-						}
+						OtherMetadata.Append(TEXT("\n"));
 					}
-
-					switch (Metadata.Type)
+					OtherMetadata.Append(Metadata.Name.GetPlainNameString());
+					OtherMetadata.Append(TEXT("="));
+					OtherMetadata.Appendf(TEXT("%f"), Metadata.DoubleValue);
+					break;
+				}
+				case TraceServices::FTraceSessionMetadata::EType::String:
+				{
+					if (OtherMetadata.Len() > 0)
 					{
-						case TraceServices::FTraceSessionMetadata::EType::Int64:
-						{
-							if (OtherMetadata.Len() > 0)
-							{
-								OtherMetadata.Append(TEXT("\n"));
-							}
-							OtherMetadata.Append(Metadata.Name.GetPlainNameString());
-							OtherMetadata.Append(TEXT("="));
-							OtherMetadata.Appendf(TEXT("%lld"), Metadata.Int64Value);
-							break;
-						}
-						case TraceServices::FTraceSessionMetadata::EType::Double:
-						{
-							if (OtherMetadata.Len() > 0)
-							{
-								OtherMetadata.Append(TEXT("\n"));
-							}
-							OtherMetadata.Append(Metadata.Name.GetPlainNameString());
-							OtherMetadata.Append(TEXT("="));
-							OtherMetadata.Appendf(TEXT("%f"), Metadata.DoubleValue);
-							break;
-						}
-						case TraceServices::FTraceSessionMetadata::EType::String:
-						{
-							if (OtherMetadata.Len() > 0)
-							{
-								OtherMetadata.Append(TEXT("\n"));
-							}
-							OtherMetadata.Append(Metadata.Name.GetPlainNameString());
-							OtherMetadata.Append(TEXT("="));
-							OtherMetadata.Append(Metadata.StringValue);
-							break;
-						}
+						OtherMetadata.Append(TEXT("\n"));
 					}
-				});
+					OtherMetadata.Append(Metadata.Name.GetPlainNameString());
+					OtherMetadata.Append(TEXT("="));
+					OtherMetadata.Append(Metadata.StringValue);
+					break;
+				}
+			}
+		});
 
-			OtherMetadataText = FText::FromString(FString(OtherMetadata.ToView()));
+	OtherMetadataText = FText::FromString(FString(OtherMetadata.ToView()));
 
-			bIsSessionInfoSet = true;
-		}
-	}
+	return true;
 }
 
 ////////////////////////////////////////////////////////////////////////////////////////////////////
diff --git a/Engine/Source/Developer/TraceInsights/Private/Insights/Widgets/SSessionInfoWindow.h b/Engine/Source/Developer/TraceInsights/Private/Insights/Widgets/SSessionInfoWindow.h
--- a/Engine/Source/Developer/TraceInsights/Private/Insights/Widgets/SSessionInfoWindow.h
+++ b/Engine/Source/Developer/TraceInsights/Private/Insights/Widgets/SSessionInfoWindow.h
@@ -137,6 +137,9 @@ private:
 	TSharedRef<SDockTab> SpawnTab_SessionInfo(const FSpawnTabArgs& Args);
 	void OnSessionInfoTabClosed(TSharedRef<SDockTab> TabBeingClosed);
 
+	void UpdateStreamTypeAndSessionName(const TraceServices::IAnalysisSession& Session);
+	bool TryUpdateDiagnosticsInfo(const TraceServices::IAnalysisSession& Session);
+
 	FText GetStreamTypeText() const { return StreamTypeText; }
 	FText GetSessionNameText() const { return SessionNameText; }
 	FText GetUriText() const { return UriText; }
