[Insights] Timing View: Added more control for the maximum depth limit affecting the Timing Events tracks. Previously it was limited to only control with predefined values 1, 4 and 1000/unlimited. Now user can set from UI (View Modes drop down menu) any value between 1 and 1000. Was added a slider and predefined values (MAX/1000, 1, 2, 4, 8, 16). X key shortcut will iterate through the predefined values. The setting is now persistent in the ini file.

diff --git a/Engine/Config/BaseUnrealInsightsSettings.ini b/Engine/Config/BaseUnrealInsightsSettings.ini
--- a/Engine/Config/BaseUnrealInsightsSettings.ini
+++ b/Engine/Config/BaseUnrealInsightsSettings.ini
@@ -20,6 +20,17 @@ bAutoHideEmptyTracks=True
 ; This option can be changed from the "View Mode" drop down menu.
 bAllowPanningOnScreenEdges=False
 
+; The minimum number of lanes in a timing events track.
+; Valid values: [0 .. 100]
+; If >0 an empty track will still be visible even if bAutoHideEmptyTracks is true.
+; This option can be changed from the "View Mode" drop down menu.
+TimingTrackDepthLimit.Min=0
+
+; The maximum number of lanes in a timing events track.
+; Valid values: [1 .. 1000]
+; This option can be changed from the "View Mode" drop down menu.
+TimingTrackDepthLimit.Max=1000
+
 ; The default state for Auto-Scroll, in the Timing view.
 bAutoScroll=False
 
diff --git a/Engine/Source/Developer/TraceInsights/Private/Insights/InsightsSettings.cpp b/Engine/Source/Developer/TraceInsights/Private/Insights/InsightsSettings.cpp
--- a/Engine/Source/Developer/TraceInsights/Private/Insights/InsightsSettings.cpp
+++ b/Engine/Source/Developer/TraceInsights/Private/Insights/InsightsSettings.cpp
@@ -122,6 +122,8 @@ void FInsightsSettings::LoadFromConfig()
 	GConfig->GetBool(TEXT("Insights.TimingProfiler"), TEXT("bAutoHideEmptyTracks"), bAutoHideEmptyTracks, SettingsIni);
 	GConfig->GetBool(TEXT("Insights.TimingProfiler"), TEXT("bAllowPanningOnScreenEdges"), bAllowPanningOnScreenEdges, SettingsIni);
 	GConfig->GetDouble(TEXT("Insights.TimingProfiler"), TEXT("ThreadsToUseRatio"), ThreadsToUseRatio, SettingsIni);
+	GConfig->GetInt(TEXT("Insights.TimingProfiler"), TEXT("TimingTrackDepthLimit.Min"), TimingTrackDepthLimitMin, SettingsIni);
+	GConfig->GetInt(TEXT("Insights.TimingProfiler"), TEXT("TimingTrackDepthLimit.Max"), TimingTrackDepthLimitMax, SettingsIni);
 
 	// Auto-scroll options
 	GConfig->GetBool(TEXT("Insights.TimingProfiler"), TEXT("bAutoScroll"), bAutoScroll, SettingsIni);
@@ -298,6 +300,8 @@ void FInsightsSettings::SaveToConfig()
 	GConfig->SetBool(TEXT("Insights.TimingProfiler"), TEXT("bAutoHideEmptyTracks"), bAutoHideEmptyTracks, SettingsIni);
 	GConfig->SetBool(TEXT("Insights.TimingProfiler"), TEXT("bAllowPanningOnScreenEdges"), bAllowPanningOnScreenEdges, SettingsIni);
 	GConfig->SetDouble(TEXT("Insights.TimingProfiler"), TEXT("ThreadsToUseRatio"), ThreadsToUseRatio, SettingsIni);
+	GConfig->SetInt(TEXT("Insights.TimingProfiler"), TEXT("TimingTrackDepthLimit.Min"), TimingTrackDepthLimitMin, SettingsIni);
+	GConfig->SetInt(TEXT("Insights.TimingProfiler"), TEXT("TimingTrackDepthLimit.Max"), TimingTrackDepthLimitMax, SettingsIni);
 
 	// Auto-scroll options
 	GConfig->SetBool(TEXT("Insights.TimingProfiler"), TEXT("bAutoScroll"), bAutoScroll, SettingsIni);
diff --git a/Engine/Source/Developer/TraceInsights/Private/Insights/InsightsSettings.h b/Engine/Source/Developer/TraceInsights/Private/Insights/InsightsSettings.h
--- a/Engine/Source/Developer/TraceInsights/Private/Insights/InsightsSettings.h
+++ b/Engine/Source/Developer/TraceInsights/Private/Insights/InsightsSettings.h
@@ -344,6 +344,14 @@ public:
 	void SetPanningOnScreenEdges(bool bOnOff) { bAllowPanningOnScreenEdges = bOnOff; }
 	void SetAndSavePanningOnScreenEdges(bool bOnOff) { SET_AND_SAVE(bAllowPanningOnScreenEdges, bOnOff); }
 
+	int32 GetTimingTrackDepthLimitMin() const { return TimingTrackDepthLimitMin; }
+	void SetTimingTrackDepthLimitMin(int32 DepthLimit) { TimingTrackDepthLimitMin = DepthLimit; }
+	void SetAndSaveTimingTrackDepthLimitMin(int32 DepthLimit) { SET_AND_SAVE(TimingTrackDepthLimitMin, DepthLimit); }
+
+	int32 GetTimingTrackDepthLimitMax() const { return TimingTrackDepthLimitMax; }
+	void SetTimingTrackDepthLimitMax(int32 DepthLimit) { TimingTrackDepthLimitMax = DepthLimit; }
+	void SetAndSaveTimingTrackDepthLimitMax(int32 DepthLimit) { SET_AND_SAVE(TimingTrackDepthLimitMax, DepthLimit); }
+
 	bool IsAutoScrollEnabled() const { return bAutoScroll; }
 	void SetAutoScroll(bool bOnOff) { bAutoScroll = bOnOff; }
 	void SetAndSaveAutoScroll(bool bOnOff) { SET_AND_SAVE(bAutoScroll, bOnOff); }
@@ -551,6 +559,12 @@ private:
 	/** If enabled, the panning is allowed to continue when mouse cursor reaches the edges of the screen. */
 	bool bAllowPanningOnScreenEdges = false;
 
+	/** The minimum number of lanes in a timing events track. */
+	int32 TimingTrackDepthLimitMin;
+
+	/** The maximum number of lanes in a timing events track. */
+	int32 TimingTrackDepthLimitMax;
+
 	/** If enabled, the Timing View will start with auto-scroll enabled. */
 	bool bAutoScroll = false;
 
diff --git a/Engine/Source/Developer/TraceInsights/Private/Insights/TimingProfiler/TimingProfilerManager.cpp b/Engine/Source/Developer/TraceInsights/Private/Insights/TimingProfiler/TimingProfilerManager.cpp
--- a/Engine/Source/Developer/TraceInsights/Private/Insights/TimingProfiler/TimingProfilerManager.cpp
+++ b/Engine/Source/Developer/TraceInsights/Private/Insights/TimingProfiler/TimingProfilerManager.cpp
@@ -14,6 +14,7 @@
 // TraceInsights
 #include "Insights/Common/InsightsMenuBuilder.h"
 #include "Insights/InsightsManager.h"
+#include "Insights/InsightsSettings.h"
 #include "Insights/InsightsStyle.h"
 #include "Insights/TimingProfiler/ViewModels/TimerButterflyAggregator.h"
 #include "Insights/TimingProfiler/ViewModels/TimingExporter.h"
@@ -109,6 +110,10 @@ void FTimingProfilerManager::Initialize(IUnrealInsightsModule& InsightsModule)
 	OnTick = FTickerDelegate::CreateSP(this, &FTimingProfilerManager::Tick);
 	OnTickHandle = FTSTicker::GetCoreTicker().AddTicker(OnTick, 0.0f);
 
+	const FInsightsSettings& Settings = FInsightsManager::Get()->GetSettings();
+	EventDepthLimitMin = Settings.GetTimingTrackDepthLimitMin();
+	EventDepthLimitMax = Settings.GetTimingTrackDepthLimitMax();
+
 	FTimingProfilerCommands::Register();
 	BindCommands();
 
@@ -1229,6 +1234,28 @@ bool FTimingProfilerManager::Exec(const TCHAR* Cmd, FOutputDevice& Ar)
 
 ////////////////////////////////////////////////////////////////////////////////////////////////////
 
+void FTimingProfilerManager::SetEventDepthLimitMin(uint32 InEventDepthLimitMin)
+{
+	EventDepthLimitMin = FMath::Min<uint32>(InEventDepthLimitMin, MaxEventDepthLimitMin);
+
+	// Persistent option. Save it to the config file.
+	FInsightsSettings& Settings = FInsightsManager::Get()->GetSettings();
+	Settings.SetAndSaveTimingTrackDepthLimitMin(static_cast<int32>(EventDepthLimitMin));
+}
+
+////////////////////////////////////////////////////////////////////////////////////////////////////
+
+void FTimingProfilerManager::SetEventDepthLimitMax(uint32 InEventDepthLimitMax)
+{
+	EventDepthLimitMax = FMath::Clamp<uint32>(InEventDepthLimitMax, 1u, UnlimitedEventDepth);
+
+	// Persistent option. Save it to the config file.
+	FInsightsSettings& Settings = FInsightsManager::Get()->GetSettings();
+	Settings.SetAndSaveTimingTrackDepthLimitMax(static_cast<int32>(EventDepthLimitMax));
+}
+
+////////////////////////////////////////////////////////////////////////////////////////////////////
+
 } // namespace UE::Insights::TimingProfiler
 
 #undef LOCTEXT_NAMESPACE
diff --git a/Engine/Source/Developer/TraceInsights/Private/Insights/TimingProfiler/TimingProfilerManager.h b/Engine/Source/Developer/TraceInsights/Private/Insights/TimingProfiler/TimingProfilerManager.h
--- a/Engine/Source/Developer/TraceInsights/Private/Insights/TimingProfiler/TimingProfilerManager.h
+++ b/Engine/Source/Developer/TraceInsights/Private/Insights/TimingProfiler/TimingProfilerManager.h
@@ -42,6 +42,7 @@ class FTimingProfilerManager : public TSharedFromThis<FTimingProfilerManager>, p
 	friend class FTimingProfilerActionManager;
 
 public:
+	inline static constexpr uint32 MaxEventDepthLimitMin = 100; // maximum value for the EventDepthLimitMin
 	inline static constexpr uint32 UnlimitedEventDepth = 1000;
 
 public:
@@ -168,8 +169,12 @@ public:
 	ETimingEventsColoringMode GetColoringMode() const { return ColoringMode; }
 	void SetColoringMode(ETimingEventsColoringMode InColoringMode) { ColoringMode = InColoringMode; }
 
-	uint32 GetEventDepthLimit() const { return EventDepthLimit; }
-	void SetEventDepthLimit(uint32 InEventDepthLimit) { EventDepthLimit = InEventDepthLimit; }
+	uint32 GetEventDepthLimitMin() const { return EventDepthLimitMin; }
+	void SetEventDepthLimitMin(uint32 InEventDepthLimit);
+
+	uint32 GetEventDepthLimitMax() const { return EventDepthLimitMax; }
+	uint32 GetEventDepthLimit() const { return EventDepthLimitMax; } // for backward compatibility
+	void SetEventDepthLimitMax(uint32 InEventDepthLimit);
 
 private:
 	/** Binds our UI commands to delegates. */
@@ -259,7 +264,8 @@ private:
 	FName LogListingName;
 
 	ETimingEventsColoringMode ColoringMode = ETimingEventsColoringMode::ByTimerName;
-	uint32 EventDepthLimit = UnlimitedEventDepth;
+	uint32 EventDepthLimitMin = 0;
+	uint32 EventDepthLimitMax = UnlimitedEventDepth;
 
 	/** A shared pointer to the global instance of the Timing Profiler manager. */
 	static TSharedPtr<FTimingProfilerManager> Instance;
diff --git a/Engine/Source/Developer/TraceInsights/Private/Insights/Widgets/STimingView.cpp b/Engine/Source/Developer/TraceInsights/Private/Insights/Widgets/STimingView.cpp
--- a/Engine/Source/Developer/TraceInsights/Private/Insights/Widgets/STimingView.cpp
+++ b/Engine/Source/Developer/TraceInsights/Private/Insights/Widgets/STimingView.cpp
@@ -29,6 +29,7 @@
 #include "Widgets/Input/SComboButton.h"
 #include "Widgets/Input/SEditableTextBox.h"
 #include "Widgets/Input/SSegmentedControl.h"
+#include "Widgets/Input/SSpinBox.h"
 #include "Widgets/Layout/SScrollBar.h"
 #include "Widgets/Layout/SSpacer.h"
 #include "Widgets/SOverlay.h"
@@ -3020,7 +3021,7 @@ FReply STimingView::OnKeyDown(const FGeometry& MyGeometry, const FKeyEvent& InKe
 	}
 	else if (InKeyEvent.GetKey() == EKeys::X)
 	{
-		ChooseNextEventDepthLimit();
+		ChooseNextEventDepthLimitMax();
 		return FReply::Handled();
 	}
 #if INSIGHTS_ACTIVATE_BENCHMARK
@@ -4998,97 +4999,214 @@ void STimingView::TogglePanningOnScreenEdges()
 
 void STimingView::CreateDepthLimitMenu(FMenuBuilder& MenuBuilder)
 {
-	MenuBuilder.BeginSection("DepthLimit", LOCTEXT("ContextMenu_Section_DepthLimit", "Depth Limit"));
+	MenuBuilder.BeginSection("TrackDepthLimits", LOCTEXT("ContextMenu_Section_TrackDepthLimits", "Track Depth Limits"));
 	{
-		// Note: We use the custom FInsightsMenuBuilder::AddMenuEntry in order to set the same key binding text for multiple menu items.
-
-		FInsightsMenuBuilder::AddMenuEntry(MenuBuilder,
-			FUIAction(
-				FExecuteAction::CreateSP(this, &STimingView::SetEventDepthLimit, FTimingProfilerManager::UnlimitedEventDepth),
-				FCanExecuteAction(),
-				FIsActionChecked::CreateSP(this, &STimingView::CheckEventDepthLimit, FTimingProfilerManager::UnlimitedEventDepth)),
-			LOCTEXT("UnlimitedDepth", "Unlimited"),
-			LOCTEXT("UnlimitedDepth_Desc", "Timing Events tracks (like the CPU Thread tracks) can have unlimited depth (lanes)."),
-			TAttribute<FText>::Create(TAttribute<FText>::FGetter::CreateSP(this, &STimingView::GetEventDepthLimitKeybindingText, FTimingProfilerManager::UnlimitedEventDepth)),
-			EUserInterfaceActionType::RadioButton);
-
-		FInsightsMenuBuilder::AddMenuEntry(MenuBuilder,
-			FUIAction(
-				FExecuteAction::CreateSP(this, &STimingView::SetEventDepthLimit, (uint32)4),
-				FCanExecuteAction(),
-				FIsActionChecked::CreateSP(this, &STimingView::CheckEventDepthLimit, (uint32)4)),
-			LOCTEXT("DepthLimit4", "4 Lanes"),
-			LOCTEXT("DepthLimit4_Desc", "Timing Events tracks (like the CPU Thread tracks) can have maximum 4 lanes."),
-			TAttribute<FText>::Create(TAttribute<FText>::FGetter::CreateSP(this, &STimingView::GetEventDepthLimitKeybindingText, (uint32)4)),
-			EUserInterfaceActionType::RadioButton);
+#if 0 // TODO
+		MenuBuilder.AddMenuEntry(
+			FUIAction(),
+			SNew(SHorizontalBox)
+			+ SHorizontalBox::Slot()
+			.HAlign(HAlign_Left)
+			.VAlign(VAlign_Center)
+			.AutoWidth()
+			.Padding(FMargin(-32.0f, 0.0f, 0.0f, 0.0f))
+			[
+				SNew(STextBlock)
+				.MinDesiredWidth(34.0f)
+				.Text(LOCTEXT("TrackDepthLimitMin_Text", "Min"))
+			]
+			+ SHorizontalBox::Slot()
+			.HAlign(HAlign_Left)
+			.AutoWidth()
+			.Padding(FMargin(0.0f, 0.0f, 4.0f, 0.0f))
+			[
+				SNew(SSegmentedControl<uint32>)
+				.Value_Lambda([this]
+				{
+					return FTimingProfilerManager::Get()->GetEventDepthLimitMin();
+				})
+				.OnValueChanged_Lambda([this](uint32 InValue)
+				{
+					SetEventDepthLimitMin(InValue);
+				})
+				+ SSegmentedControl<uint32>::Slot(0u)
+				.Text(LOCTEXT("DepthLimitMin_0", "0"))
+				+ SSegmentedControl<uint32>::Slot(1u)
+				.Text(LOCTEXT("DepthLimitMin_1", "1"))
+				+ SSegmentedControl<uint32>::Slot(2u)
+				.Text(LOCTEXT("DepthLimitMin_2", "2"))
+				+ SSegmentedControl<uint32>::Slot(4u)
+				.Text(LOCTEXT("DepthLimitMin_4", "4"))
+				+ SSegmentedControl<uint32>::Slot(8u)
+				.Text(LOCTEXT("DepthLimitMin_8", "8"))
+				+ SSegmentedControl<uint32>::Slot(16u)
+				.Text(LOCTEXT("DepthLimitMin_16", "16"))
+			]
+			+ SHorizontalBox::Slot()
+			.HAlign(HAlign_Left)
+			.AutoWidth()
+			.Padding(FMargin(0.0f, 0.0f, 8.0f, 0.0f))
+			[
+				SNew(SSpinBox<uint32>)
+				.MinValue(TOptional<uint32>(0))
+				.MaxValue(TOptional<uint32>(100))
+				.MinSliderValue(TOptional<uint32>(0))
+				.MaxSliderValue(TOptional<uint32>(16))
+				.Delta(1)
+				.MinDesiredWidth(60.0f)
+				.Value_Lambda([this]() -> uint32
+				{
+					return FTimingProfilerManager::Get()->GetEventDepthLimitMin();
+				})
+				.OnValueChanged_Lambda([this](uint32 InValue)
+				{
+					SetEventDepthLimitMin(InValue);
+				})
+			],
+			NAME_None,
+			LOCTEXT("TrackDepthLimitMin_Tooltip", "The minimum depth (number of lanes) for a Timing Events track (like a CPU Thread track).\nValid values: [0 .. 100]\nThis option is persistent to the UnrealInsightsSettings.ini file."),
+			EUserInterfaceActionType::None
+		);
+#endif
 
-		FInsightsMenuBuilder::AddMenuEntry(MenuBuilder,
-			FUIAction(
-				FExecuteAction::CreateSP(this, &STimingView::SetEventDepthLimit, (uint32)1),
-				FCanExecuteAction(),
-				FIsActionChecked::CreateSP(this, &STimingView::CheckEventDepthLimit, (uint32)1)),
-			LOCTEXT("DepthLimit1", "Single Lane"),
-			LOCTEXT("DepthLimit1_Desc", "Timing Events tracks (like the CPU Thread tracks) can have a single lane."),
-			TAttribute<FText>::Create(TAttribute<FText>::FGetter::CreateSP(this, &STimingView::GetEventDepthLimitKeybindingText, (uint32)1)),
-			EUserInterfaceActionType::RadioButton);
+		MenuBuilder.AddMenuEntry(
+			FUIAction(),
+			SNew(SHorizontalBox)
+			+ SHorizontalBox::Slot()
+			.HAlign(HAlign_Left)
+			.VAlign(VAlign_Center)
+			.AutoWidth()
+			.Padding(FMargin(-32.0f, 0.0f, 0.0f, 0.0f))
+			[
+				SNew(STextBlock)
+				.MinDesiredWidth(34.0f)
+				.Text(LOCTEXT("TrackDepthLimitMax_Text", "Max"))
+			]
+			+ SHorizontalBox::Slot()
+			.HAlign(HAlign_Left)
+			.AutoWidth()
+			.Padding(FMargin(0.0f, 0.0f, 4.0f, 0.0f))
+			[
+				SNew(SSegmentedControl<uint32>)
+				.Value_Lambda([this]
+				{
+					return FTimingProfilerManager::Get()->GetEventDepthLimitMax();
+				})
+				.OnValueChanged_Lambda([this](uint32 InValue)
+				{
+					SetEventDepthLimitMax(InValue);
+				})
+				+ SSegmentedControl<uint32>::Slot(1000u)
+				[
+					SNew(STextBlock)
+					.Font(FAppStyle::Get().GetFontStyle("SmallFont"))
+					.Margin(FMargin(-4.0f, 2.0f, -4.0f, 1.0f))
+					.Text(LOCTEXT("DepthLimitMax_Max", "MAX"))
+				]
+				+ SSegmentedControl<uint32>::Slot(1u)
+				.Text(LOCTEXT("DepthLimitMax_1", "1"))
+				+ SSegmentedControl<uint32>::Slot(2u)
+				.Text(LOCTEXT("DepthLimitMax_2", "2"))
+				+ SSegmentedControl<uint32>::Slot(4u)
+				.Text(LOCTEXT("DepthLimitMax_4", "4"))
+				+ SSegmentedControl<uint32>::Slot(8u)
+				.Text(LOCTEXT("DepthLimitMax_8", "8"))
+				+ SSegmentedControl<uint32>::Slot(16u)
+				.Text(LOCTEXT("DepthLimitMax_16", "16"))
+			]
+			+ SHorizontalBox::Slot()
+			.HAlign(HAlign_Left)
+			.AutoWidth()
+			.Padding(FMargin(0.0f, 0.0f, 4.0f, 0.0f))
+			[
+				SNew(SSpinBox<uint32>)
+				.MinValue(TOptional<uint32>(0))
+				.MaxValue(TOptional<uint32>(100))
+				.MinSliderValue(TOptional<uint32>(0))
+				.MaxSliderValue(TOptional<uint32>(16))
+				.Delta(1)
+				.MinDesiredWidth(60.0f)
+				.Value_Lambda([this]() -> uint32
+				{
+					return FTimingProfilerManager::Get()->GetEventDepthLimitMax();
+				})
+				.OnValueChanged_Lambda([this](uint32 InValue)
+				{
+					SetEventDepthLimitMax(InValue);
+				})
+			]
+			+ SHorizontalBox::Slot()
+			.Padding(FMargin(4.0f, 0.0f))
+			.HAlign(HAlign_Right)
+			.VAlign(VAlign_Center)
+			.FillWidth(1.0f)
+			[
+				SNew(STextBlock)
+				.TextStyle(FAppStyle::Get(), "Menu.Keybinding")
+				.ColorAndOpacity(FSlateColor::UseSubduedForeground())
+				.Text(LOCTEXT("TrackDepthLimitMax_Keybinding", "X"))
+			],
+			NAME_None,
+			LOCTEXT("TrackDepthLimitMax_Tooltip", "The maximum depth (number of lanes) for a Timing Events track (like a CPU Thread track).\nValid values: [1 .. 1000]\nThis option is persistent to the UnrealInsightsSettings.ini file."),
+			EUserInterfaceActionType::None
+		);
 	}
 	MenuBuilder.EndSection();
 }
 
 ////////////////////////////////////////////////////////////////////////////////////////////////////
 
-FText STimingView::GetEventDepthLimitKeybindingText(uint32 DepthLimit) const
+FText STimingView::GetEventDepthLimitKeybindingText(uint32 InDepthLimitMax) const
 {
-	uint32 CurrentDepthLimit = FTimingProfilerManager::Get()->GetEventDepthLimit();
-	uint32 NextDepthLimit = GetNextEventDepthLimit(CurrentDepthLimit);
-	if (DepthLimit == NextDepthLimit)
-	{
-		return FInputChord(EKeys::X).GetInputText().ToUpper();
-	}
-	return FText::GetEmpty();
+	uint32 CurrentDepthLimit = FTimingProfilerManager::Get()->GetEventDepthLimitMax();
+	uint32 NextDepthLimit = GetNextEventDepthLimitMax(CurrentDepthLimit);
+	return InDepthLimitMax == NextDepthLimit ? LOCTEXT("DepthLimitKeybinding", "X") : FText::GetEmpty();
 }
 
 ////////////////////////////////////////////////////////////////////////////////////////////////////
 
-uint32 STimingView::GetNextEventDepthLimit(uint32 DepthLimit) const
+bool STimingView::CheckEventDepthLimitMax(uint32 InDepthLimitMax) const
 {
-	if (DepthLimit == 1)
-	{
-		return 4;
-	}
-	else if (DepthLimit == 4)
-	{
-		return FTimingProfilerManager::UnlimitedEventDepth;
-	}
-	else // Unlimited
+	return InDepthLimitMax == FTimingProfilerManager::Get()->GetEventDepthLimitMax();
+}
+
+////////////////////////////////////////////////////////////////////////////////////////////////////
+
+uint32 STimingView::GetNextEventDepthLimitMax(uint32 DepthLimit) const
+{
+	// Cycles between: 1, 2, 4, 8, 16, Unlimited
+	if (DepthLimit == FTimingProfilerManager::UnlimitedEventDepth)
 	{
 		return 1;
 	}
+	DepthLimit = FMath::RoundUpToPowerOfTwo(DepthLimit * 2);
+	return (DepthLimit <= 16) ? DepthLimit : FTimingProfilerManager::UnlimitedEventDepth;
 }
 
 ////////////////////////////////////////////////////////////////////////////////////////////////////
 
-void STimingView::ChooseNextEventDepthLimit()
+void STimingView::ChooseNextEventDepthLimitMax()
 {
-	const uint32 CurrentDepthLimit = FTimingProfilerManager::Get()->GetEventDepthLimit();
-	const uint32 NextDepthLimit = GetNextEventDepthLimit(CurrentDepthLimit);
-	FTimingProfilerManager::Get()->SetEventDepthLimit(NextDepthLimit);
+	const uint32 CurrentDepthLimit = FTimingProfilerManager::Get()->GetEventDepthLimitMax();
+	const uint32 NextDepthLimit = GetNextEventDepthLimitMax(CurrentDepthLimit);
+	FTimingProfilerManager::Get()->SetEventDepthLimitMax(NextDepthLimit);
 	Viewport.AddDirtyFlags(ETimingTrackViewportDirtyFlags::HInvalidated);
 }
 
 ////////////////////////////////////////////////////////////////////////////////////////////////////
 
-void STimingView::SetEventDepthLimit(uint32 DepthLimit)
+void STimingView::SetEventDepthLimitMax(uint32 InDepthLimitMax)
 {
-	FTimingProfilerManager::Get()->SetEventDepthLimit(DepthLimit);
+	FTimingProfilerManager::Get()->SetEventDepthLimitMax(InDepthLimitMax);
 	Viewport.AddDirtyFlags(ETimingTrackViewportDirtyFlags::HInvalidated);
 }
 
 ////////////////////////////////////////////////////////////////////////////////////////////////////
 
-bool STimingView::CheckEventDepthLimit(uint32 DepthLimit) const
+void STimingView::SetEventDepthLimitMin(uint32 InDepthLimitMin)
 {
-	return DepthLimit == FTimingProfilerManager::Get()->GetEventDepthLimit();
+	FTimingProfilerManager::Get()->SetEventDepthLimitMin(InDepthLimitMin);
+	Viewport.AddDirtyFlags(ETimingTrackViewportDirtyFlags::HInvalidated);
 }
 
 ////////////////////////////////////////////////////////////////////////////////////////////////////
diff --git a/Engine/Source/Developer/TraceInsights/Private/Insights/Widgets/STimingView.h b/Engine/Source/Developer/TraceInsights/Private/Insights/Widgets/STimingView.h
--- a/Engine/Source/Developer/TraceInsights/Private/Insights/Widgets/STimingView.h
+++ b/Engine/Source/Developer/TraceInsights/Private/Insights/Widgets/STimingView.h
@@ -468,11 +468,12 @@ protected:
 	TSharedRef<SWidget> MakeViewModeMenu();
 
 	void CreateDepthLimitMenu(FMenuBuilder& MenuBuilder);
-	FText GetEventDepthLimitKeybindingText(uint32 DepthLimit) const;
-	uint32 GetNextEventDepthLimit(uint32 DepthLimit) const;
-	void ChooseNextEventDepthLimit();
-	void SetEventDepthLimit(uint32 DepthLimit);
-	bool CheckEventDepthLimit(uint32 DepthLimit) const;
+	FText GetEventDepthLimitKeybindingText(uint32 InDepthLimitMax) const;
+	uint32 GetNextEventDepthLimitMax(uint32 InDepthLimitMax) const;
+	bool CheckEventDepthLimitMax(uint32 InDepthLimitMax) const;
+	void ChooseNextEventDepthLimitMax();
+	void SetEventDepthLimitMax(uint32 InDepthLimitMax);
+	void SetEventDepthLimitMin(uint32 InDepthLimitMin);
 
 	void CreateCpuThreadTrackColoringModeMenu(FMenuBuilder& MenuBuilder);
 	void ChooseNextCpuThreadTrackColoringMode();
@@ -771,15 +772,16 @@ protected:
 
 	// Debug stats
 	int32 NumUpdatedEvents;
-	TFixedCircularBuffer<uint64, 32> PreUpdateTracksDurationHistory;
-	TFixedCircularBuffer<uint64, 32> UpdateTracksDurationHistory;
-	TFixedCircularBuffer<uint64, 32> PostUpdateTracksDurationHistory;
-	TFixedCircularBuffer<uint64, 32> TickDurationHistory;
-	mutable TFixedCircularBuffer<uint64, 32> PreDrawTracksDurationHistory;
-	mutable TFixedCircularBuffer<uint64, 32> DrawTracksDurationHistory;
-	mutable TFixedCircularBuffer<uint64, 32> PostDrawTracksDurationHistory;
-	mutable TFixedCircularBuffer<uint64, 32> TotalDrawDurationHistory;
-	mutable TFixedCircularBuffer<uint64, 32> OnPaintDeltaTimeHistory;
+	typedef TFixedCircularBuffer<uint64, 32> FHistoryBuffer;
+	FHistoryBuffer PreUpdateTracksDurationHistory;
+	FHistoryBuffer UpdateTracksDurationHistory;
+	FHistoryBuffer PostUpdateTracksDurationHistory;
+	FHistoryBuffer TickDurationHistory;
+	mutable FHistoryBuffer PreDrawTracksDurationHistory;
+	mutable FHistoryBuffer DrawTracksDurationHistory;
+	mutable FHistoryBuffer PostDrawTracksDurationHistory;
+	mutable FHistoryBuffer TotalDrawDurationHistory;
+	mutable FHistoryBuffer OnPaintDeltaTimeHistory;
 	mutable uint64 LastOnPaintTime;
 
 	////////////////////////////////////////////////////////////
