[Insights] Session Info: Fixed potential unitialized UriText when session changes.

#rnx

diff --git a/Engine/Source/Developer/TraceInsights/Private/Insights/Widgets/SSessionInfoWindow.cpp b/Engine/Source/Developer/TraceInsights/Private/Insights/Widgets/SSessionInfoWindow.cpp
--- a/Engine/Source/Developer/TraceInsights/Private/Insights/Widgets/SSessionInfoWindow.cpp
+++ b/Engine/Source/Developer/TraceInsights/Private/Insights/Widgets/SSessionInfoWindow.cpp
@@ -361,6 +361,7 @@ void SSessionInfoWindow::Tick(const FGeometry& AllottedGeometry, const double In
 
 		StreamTypeText = FText::GetEmpty();
 		SessionNameText = FText::GetEmpty();
+		UriText = FText::GetEmpty();
 
 		PlatformText = FText::GetEmpty();
 		AppNameText = FText::GetEmpty();
@@ -433,6 +434,7 @@ void SSessionInfoWindow::UpdateStreamTypeAndSessionName(const TraceServices::IAn
 	else
 	{
 		SessionNameText = FText::FromString(Session.GetName());
+		UriText = FText::GetEmpty();
 	}
 }
 
