[Insights] Fixed max depth limit to accept values between 100 and 1000. Also changed "MAX" text to infinity sign.

#rnx

diff --git a/Engine/Source/Developer/TraceInsights/Private/Insights/Widgets/STimingView.cpp b/Engine/Source/Developer/TraceInsights/Private/Insights/Widgets/STimingView.cpp
--- a/Engine/Source/Developer/TraceInsights/Private/Insights/Widgets/STimingView.cpp
+++ b/Engine/Source/Developer/TraceInsights/Private/Insights/Widgets/STimingView.cpp
@@ -5030,17 +5030,17 @@ void STimingView::CreateDepthLimitMenu(FMenuBuilder& MenuBuilder)
 					SetEventDepthLimitMin(InValue);
 				})
 				+ SSegmentedControl<uint32>::Slot(0u)
-				.Text(LOCTEXT("DepthLimitMin_0", "0"))
+				.Text(FText::FromString("0"))
 				+ SSegmentedControl<uint32>::Slot(1u)
-				.Text(LOCTEXT("DepthLimitMin_1", "1"))
+				.Text(FText::FromString("1"))
 				+ SSegmentedControl<uint32>::Slot(2u)
-				.Text(LOCTEXT("DepthLimitMin_2", "2"))
+				.Text(FText::FromString("2"))
 				+ SSegmentedControl<uint32>::Slot(4u)
-				.Text(LOCTEXT("DepthLimitMin_4", "4"))
+				.Text(FText::FromString("4"))
 				+ SSegmentedControl<uint32>::Slot(8u)
-				.Text(LOCTEXT("DepthLimitMin_8", "8"))
+				.Text(FText::FromString("8"))
 				+ SSegmentedControl<uint32>::Slot(16u)
-				.Text(LOCTEXT("DepthLimitMin_16", "16"))
+				.Text(FText::FromString("16"))
 			]
 			+ SHorizontalBox::Slot()
 			.HAlign(HAlign_Left)
@@ -5099,20 +5099,19 @@ void STimingView::CreateDepthLimitMenu(FMenuBuilder& MenuBuilder)
 				+ SSegmentedControl<uint32>::Slot(1000u)
 				[
 					SNew(STextBlock)
-					.Font(FAppStyle::Get().GetFontStyle("SmallFont"))
-					.Margin(FMargin(-4.0f, 2.0f, -4.0f, 1.0f))
-					.Text(LOCTEXT("DepthLimitMax_Max", "MAX"))
+					.Margin(FMargin(0.0f, -1.0f, 0.0f, 0.0f))
+					.Text(LOCTEXT("DepthLimitMax_Unlimited", "\u221E"))
 				]
 				+ SSegmentedControl<uint32>::Slot(1u)
-				.Text(LOCTEXT("DepthLimitMax_1", "1"))
+				.Text(FText::FromString("1"))
 				+ SSegmentedControl<uint32>::Slot(2u)
-				.Text(LOCTEXT("DepthLimitMax_2", "2"))
+				.Text(FText::FromString("2"))
 				+ SSegmentedControl<uint32>::Slot(4u)
-				.Text(LOCTEXT("DepthLimitMax_4", "4"))
+				.Text(FText::FromString("4"))
 				+ SSegmentedControl<uint32>::Slot(8u)
-				.Text(LOCTEXT("DepthLimitMax_8", "8"))
+				.Text(FText::FromString("8"))
 				+ SSegmentedControl<uint32>::Slot(16u)
-				.Text(LOCTEXT("DepthLimitMax_16", "16"))
+				.Text(FText::FromString("16"))
 			]
 			+ SHorizontalBox::Slot()
 			.HAlign(HAlign_Left)
@@ -5120,9 +5119,9 @@ void STimingView::CreateDepthLimitMenu(FMenuBuilder& MenuBuilder)
 			.Padding(FMargin(0.0f, 0.0f, 4.0f, 0.0f))
 			[
 				SNew(SSpinBox<uint32>)
-				.MinValue(TOptional<uint32>(0))
-				.MaxValue(TOptional<uint32>(100))
-				.MinSliderValue(TOptional<uint32>(0))
+				.MinValue(TOptional<uint32>(1))
+				.MaxValue(TOptional<uint32>(1000))
+				.MinSliderValue(TOptional<uint32>(1))
 				.MaxSliderValue(TOptional<uint32>(16))
 				.Delta(1)
 				.MinDesiredWidth(60.0f)
