[Insights] Timing View: Improved performance for the Quick Find operations with filtering by string values (now it is ~2.7 times faster; ex. a search that took 47.3 seconds now takes 17.3 seconds).

diff --git a/Engine/Source/Developer/TraceInsights/Private/Insights/TimingProfiler/ViewModels/TimerFilters.cpp b/Engine/Source/Developer/TraceInsights/Private/Insights/TimingProfiler/ViewModels/TimerFilters.cpp
--- a/Engine/Source/Developer/TraceInsights/Private/Insights/TimingProfiler/ViewModels/TimerFilters.cpp
+++ b/Engine/Source/Developer/TraceInsights/Private/Insights/TimingProfiler/ViewModels/TimerFilters.cpp
@@ -117,9 +117,9 @@ FTimerNameFilter::FTimerNameFilter()
 {
 	SupportedOperators = MakeShared<TArray<TSharedPtr<IFilterOperator>>>();
 	SupportedOperators->Add(StaticCastSharedRef<IFilterOperator>(MakeShared<FFilterOperator<int64>>(
-		EFilterOperator::Eq, TEXT("Is"), [](int64 lhs, int64 rhs) { return lhs == rhs; })));
+		EFilterOperator::Eq, TEXT("IS"), [](int64 lhs, int64 rhs) { return lhs == rhs; })));
 	SupportedOperators->Add(StaticCastSharedRef<IFilterOperator>(MakeShared<FFilterOperator<int64>>(
-		EFilterOperator::Contains, TEXT("Contains"), [](int64 lhs, int64 rhs) { return lhs == rhs; })));
+		EFilterOperator::Contains, TEXT("CONTAINS"), [](int64 lhs, int64 rhs) { return lhs == rhs; })));
 
 	SetCallback([this](const FString& Text, TArray<FString>& OutSuggestions)
 		{
@@ -193,11 +193,11 @@ void FMetadataFilterState::AddCustomUI(TSharedRef<SHorizontalBox> Box)
 		.Padding(FMargin(4.0f, 0.0f, 4.0f, 0.0f))
 		[
 			SNew(SEditableTextBox)
-			.MinDesiredWidth(200.0f)
+			.MinDesiredWidth(50.0f)
 			.OnTextCommitted(this, &FMetadataFilterState::OnKeyTextBoxValueCommitted)
 			.Text(this, &FMetadataFilterState::GetKeyTextBoxValue)
 			.ToolTipText(LOCTEXT("MetadataKeyTooltipText", "The key value for the metadata."))
-			.HintText(LOCTEXT("MetadataKeyHint", "The key value for the metadata."))
+			.HintText(LOCTEXT("MetadataKeyHint", "The metadata field name"))
 		];
 
 	Box->AddSlot()
@@ -260,7 +260,7 @@ void FMetadataFilterState::Update()
 	}
 
 	ConvertedData = ConvertedDataVariant();
-	
+
 	ensure(SelectedDataType.IsValid());
 	switch (SelectedDataType->Type)
 	{
@@ -383,7 +383,7 @@ bool FMetadataFilterState::ApplyFilterToMetadata(TArrayView<const uint8>& Metada
 			{
 				return true;
 			}
-			
+
 			continue;
 		}
 
@@ -396,7 +396,7 @@ bool FMetadataFilterState::ApplyFilterToMetadata(TArrayView<const uint8>& Metada
 
 			FString Value = Context.AsString();
 
-			using FFilterOperatorString = FFilterOperator<FString>;
+			using FFilterOperatorString = FFilterOperator<FStringView>;
 			FFilterOperatorString* Operator = (FFilterOperatorString*)SelectedOperator.Get();
 			if (Operator->Apply(Value, Term))
 			{
@@ -416,7 +416,7 @@ bool FMetadataFilterState::ApplyFilterToMetadata(TArrayView<const uint8>& Metada
 			FAnsiStringView Value(Context.AsCString(), static_cast<int32>(Context.AsLength()));
 			FString ValueStr(Value);
 
-			using FFilterOperatorString = FFilterOperator<FString>;
+			using FFilterOperatorString = FFilterOperator<FStringView>;
 			FFilterOperatorString* Operator = (FFilterOperatorString*)SelectedOperator.Get();
 			if (Operator->Apply(ValueStr, Term))
 			{
@@ -650,7 +650,7 @@ bool FMetadataFilterState::Equals(const FFilterState& Other) const
 	bool bIsEqual = true;
 
 	const FMetadataFilterState& OtherMetadataFilter = StaticCast<const FMetadataFilterState&>(Other);
-	
+
 	bIsEqual &= Key.Equals(OtherMetadataFilter.Key);
 	bIsEqual &= Term.Equals(OtherMetadataFilter.Term);
 	bIsEqual &= SelectedDataType == OtherMetadataFilter.SelectedDataType;
@@ -663,7 +663,7 @@ bool FMetadataFilterState::Equals(const FFilterState& Other) const
 TSharedRef<FFilterState> FMetadataFilterState::DeepCopy() const
 {
 	TSharedRef<FMetadataFilterState> Copy = MakeShared<FMetadataFilterState>(*this);
-	
+
 	return Copy;
 }
 
diff --git a/Engine/Source/Developer/TraceInsightsCore/Private/InsightsCore/Filter/ViewModels/Filters.cpp b/Engine/Source/Developer/TraceInsightsCore/Private/InsightsCore/Filter/ViewModels/Filters.cpp
--- a/Engine/Source/Developer/TraceInsightsCore/Private/InsightsCore/Filter/ViewModels/Filters.cpp
+++ b/Engine/Source/Developer/TraceInsightsCore/Private/InsightsCore/Filter/ViewModels/Filters.cpp
@@ -43,10 +43,10 @@ FFilterStorage::FFilterStorage()
 	IntegerOperators->Add(MakeShared<FFilterOperator<int64>>(EFilterOperator::Gte, TEXT("\u2265"), std::greater_equal<>()));
 
 	StringOperators = MakeShared<TArray<TSharedPtr<IFilterOperator>>>();
-	StringOperators->Add(MakeShared<FFilterOperator<FString>>(EFilterOperator::Eq, TEXT("IS"), [](const FString& lhs, const FString& rhs) { return lhs.Equals(rhs); }));
-	StringOperators->Add(MakeShared<FFilterOperator<FString>>(EFilterOperator::NotEq, TEXT("IS NOT"), [](const FString& lhs, const FString& rhs) { return !lhs.Equals(rhs); }));
-	StringOperators->Add(MakeShared<FFilterOperator<FString>>(EFilterOperator::Contains, TEXT("CONTAINS"), [](const FString& lhs, const FString& rhs) { return lhs.Contains(rhs); }));
-	StringOperators->Add(MakeShared<FFilterOperator<FString>>(EFilterOperator::NotContains, TEXT("NOT CONTAINS"), [](const FString& lhs, const FString& rhs) { return !lhs.Contains(rhs); }));
+	StringOperators->Add(MakeShared<FFilterOperator<FStringView>>(EFilterOperator::Eq, TEXT("IS"), [](FStringView lhs, FStringView rhs) { return lhs.Equals(rhs); }));
+	StringOperators->Add(MakeShared<FFilterOperator<FStringView>>(EFilterOperator::NotEq, TEXT("IS NOT"), [](FStringView lhs, FStringView rhs) { return !lhs.Equals(rhs); }));
+	StringOperators->Add(MakeShared<FFilterOperator<FStringView>>(EFilterOperator::Contains, TEXT("CONTAINS"), [](FStringView lhs, FStringView rhs) { return lhs.Contains(rhs); }));
+	StringOperators->Add(MakeShared<FFilterOperator<FStringView>>(EFilterOperator::NotContains, TEXT("NOT CONTAINS"), [](FStringView lhs, FStringView rhs) { return !lhs.Contains(rhs); }));
 
 	FilterGroupOperators.Add(MakeShared<FFilterGroupOperator>(EFilterGroupOperator::And, LOCTEXT("AllOf", "All Of (AND)"), LOCTEXT("AllOfDesc", "All of the children must be true for the group to return true. Equivalent to an AND operation.")));
 	FilterGroupOperators.Add(MakeShared<FFilterGroupOperator>(EFilterGroupOperator::Or, LOCTEXT("AnyOf", "Any Of (OR)"), LOCTEXT("AnyOfDesc", "Any of the children must be true for the group to return true. Equivalent to an OR operation.")));
@@ -242,7 +242,6 @@ bool FFilterState::ApplyFilter(const FFilterContext& Context) const
 		FFilterOperator<double>* Operator = (FFilterOperator<double>*) SelectedOperator.Get();
 		double Value = 0.0;
 		Context.GetFilterData<double>(Filter->GetKey(), Value);
-
 		Ret = Operator->Apply(Value, FilterValue.Get<double>());
 		break;
 	}
@@ -252,16 +251,13 @@ bool FFilterState::ApplyFilter(const FFilterContext& Context) const
 		FFilterOperator<int64>* Operator = (FFilterOperator<int64>*) SelectedOperator.Get();
 		int64 Value = 0;
 		Context.GetFilterData<int64>(Filter->GetKey(), Value);
-
 		Ret = Operator->Apply(Value, FilterValue.Get<int64>());
 		break;
 	}
 	case EFilterDataType::String:
 	{
-		FFilterOperator<FString>* Operator = (FFilterOperator<FString>*) SelectedOperator.Get();
-		FString Value;
-		Context.GetFilterData<FString>(Filter->GetKey(), Value);
-
+		FFilterOperator<FStringView>* Operator = (FFilterOperator<FStringView>*) SelectedOperator.Get();
+		const FString& Value = Context.GetFilterDataAsString(Filter->GetKey());
 		Ret = Operator->Apply(Value, FilterValue.Get<FString>());
 		break;
 	}
diff --git a/Engine/Source/Developer/TraceInsightsCore/Public/InsightsCore/Filter/ViewModels/Filters.h b/Engine/Source/Developer/TraceInsightsCore/Public/InsightsCore/Filter/ViewModels/Filters.h
--- a/Engine/Source/Developer/TraceInsightsCore/Public/InsightsCore/Filter/ViewModels/Filters.h
+++ b/Engine/Source/Developer/TraceInsightsCore/Public/InsightsCore/Filter/ViewModels/Filters.h
@@ -178,6 +178,15 @@ public:
 		OutData = Data->Get<T>();
 	}
 
+	const FString& GetFilterDataAsString(int32 Key) const
+	{
+		const ContextData* Data = DataMap.Find(Key);
+		check(Data);
+
+		check(Data->IsType<FString>());
+		return Data->Get<FString>();
+	}
+
 	bool HasFilterData(int32 Key) const
 	{
 		return DataMap.Contains(Key);
