[Insights] Add bandwidth and percent value types for graphs and counter stats.

Adds two new counter "types":
* Bandwidth - Bytes per second.
* Percent - Decimal fraction displayed as percentage.

diff --git a/Engine/Source/Developer/TraceInsights/Private/Insights/TimingProfiler/GraphTracks/CounterGraphSeries.cpp b/Engine/Source/Developer/TraceInsights/Private/Insights/TimingProfiler/GraphTracks/CounterGraphSeries.cpp
--- a/Engine/Source/Developer/TraceInsights/Private/Insights/TimingProfiler/GraphTracks/CounterGraphSeries.cpp
+++ b/Engine/Source/Developer/TraceInsights/Private/Insights/TimingProfiler/GraphTracks/CounterGraphSeries.cpp
@@ -40,8 +40,19 @@ void FCounterGraphSeries::InitFromProvider()
 					{
 						SetName(CounterName);
 					}
-					//bIsTime = (Counter.GetDisplayHint() == TraceServices::CounterDisplayHint_Time);
-					bIsMemory = (Counter.GetDisplayHint() == TraceServices::CounterDisplayHint_Memory);
+					switch(Counter.GetDisplayHint())
+					{
+					case TraceServices::CounterDisplayHint_Memory: 
+						Type = EGraphType::Memory; break;
+					case TraceServices::CounterDisplayHint_Bandwidth: 
+						Type = EGraphType::Bandwidth; break;
+					case TraceServices::CounterDisplayHint_Percent: 
+						Type = EGraphType::Percent; break;
+					// case TraceServices::CounterDisplayHint_Time: 
+					// 	Type = EGraphType::Time; break;
+					default:
+						Type = EGraphType::Dimensionless;
+					}
 					bIsFloatingPoint = Counter.IsFloatingPoint();
 				});
 		}
@@ -52,52 +63,66 @@ void FCounterGraphSeries::InitFromProvider()
 
 FString FCounterGraphSeries::FormatValue(double Value) const
 {
-	if (bIsTime)
+	switch(Type)
 	{
-		return FormatTimeAuto(Value);
-	}
-	else if (bIsMemory)
-	{
-		const int64 MemValue = static_cast<int64>(Value);
-		if (MemValue > 0)
-		{
-			if (MemValue < 1024)
+		case EGraphType::Time:
+			return FormatTimeAuto(Value);
+			break;
+
+		case EGraphType::Memory:
 			{
-				return FString::Printf(TEXT("%s bytes"), *FText::AsNumber(MemValue).ToString());
+				const int64 MemValue = static_cast<int64>(Value);
+				if (MemValue == 0)
+				{
+					return TEXT("0");
+				}
+				if (MemValue < 1024)
+				{
+					return FString::Printf(TEXT("%s bytes"), *FText::AsNumber(MemValue).ToString());
+				}
+				else
+				{
+					FNumberFormattingOptions FormattingOptions;
+					FormattingOptions.MaximumFractionalDigits = 2;
+					return FString::Printf(TEXT("%s (%s bytes)"), *FText::AsMemory(MemValue, &FormattingOptions).ToString(), *FText::AsNumber(MemValue).ToString());
+				}
 			}
-			else
+			break;
+
+		case EGraphType::Bandwidth:
+			{
+				if (Value > 1024.0)
+				{
+					FNumberFormattingOptions FormattingOptions;
+					FormattingOptions.MaximumFractionalDigits = 2;
+					return FString::Printf(TEXT("%s/s"), *FText::AsMemory(static_cast<int64>(Value), &FormattingOptions).ToString());
+				}
+				else
+				{
+					return FString::Printf(TEXT("%s bytes/s"), *FText::AsNumber(Value).ToString());
+				}
+			}
+			break;
+
+		case EGraphType::Percent:
 			{
 				FNumberFormattingOptions FormattingOptions;
 				FormattingOptions.MaximumFractionalDigits = 2;
-				return FString::Printf(TEXT("%s (%s bytes)"), *FText::AsMemory(MemValue, &FormattingOptions).ToString(), *FText::AsNumber(MemValue).ToString());
+				return FText::AsPercent(Value).ToString();
 			}
-		}
-		else if (MemValue == 0)
-		{
-			return TEXT("0");
-		}
-		else
-		{
-			if (-MemValue < 1024)
+			break;
+
+		default:
+			if (bIsFloatingPoint)
 			{
-				return FString::Printf(TEXT("-%s bytes"), *FText::AsNumber(-MemValue).ToString());
+				return FString::Printf(TEXT("%g"), Value);
 			}
 			else
 			{
-				FNumberFormattingOptions FormattingOptions;
-				FormattingOptions.MaximumFractionalDigits = 2;
-				return FString::Printf(TEXT("-%s (-%s bytes)"), *FText::AsMemory(-MemValue, &FormattingOptions).ToString(), *FText::AsNumber(-MemValue).ToString());
+				const int64 Int64Value = static_cast<int64>(Value);
+				return FText::AsNumber(Int64Value).ToString();
 			}
-		}
-	}
-	else if (bIsFloatingPoint)
-	{
-		return FString::Printf(TEXT("%g"), Value);
-	}
-	else
-	{
-		const int64 Int64Value = static_cast<int64>(Value);
-		return FText::AsNumber(Int64Value).ToString();
+			break;
 	}
 }
 
diff --git a/Engine/Source/Developer/TraceInsights/Private/Insights/TimingProfiler/GraphTracks/CounterGraphSeries.h b/Engine/Source/Developer/TraceInsights/Private/Insights/TimingProfiler/GraphTracks/CounterGraphSeries.h
--- a/Engine/Source/Developer/TraceInsights/Private/Insights/TimingProfiler/GraphTracks/CounterGraphSeries.h
+++ b/Engine/Source/Developer/TraceInsights/Private/Insights/TimingProfiler/GraphTracks/CounterGraphSeries.h
@@ -21,13 +21,21 @@ public:
 	void InitFromProvider();
 	virtual FString FormatValue(double Value) const override;
 	uint32 GetCounterId() const { return CounterId; }
-	virtual bool IsTimeUnit() const { return bIsTime; }
+	virtual bool IsTimeUnit() const { return Type == EGraphType::Time; }
 	virtual void Update(FTimingGraphTrack& GraphTrack, const FTimingTrackViewport& Viewport) override;
 
 private:
+	enum class EGraphType : uint8
+	{
+		Dimensionless,
+		Time,
+		Memory,
+		Percent,
+		Bandwidth,
+	};
+	
 	uint32 CounterId = 0;
-	bool bIsTime = false; // the unit for values is [second]
-	bool bIsMemory = false; // the unit for value is [byte]
+	EGraphType Type = EGraphType::Dimensionless;
 	bool bIsFloatingPoint = false; // for stats counters
 };
 
diff --git a/Engine/Source/Developer/TraceInsights/Private/Insights/TimingProfiler/ViewModels/StatsNode.cpp b/Engine/Source/Developer/TraceInsights/Private/Insights/TimingProfiler/ViewModels/StatsNode.cpp
--- a/Engine/Source/Developer/TraceInsights/Private/Insights/TimingProfiler/ViewModels/StatsNode.cpp
+++ b/Engine/Source/Developer/TraceInsights/Private/Insights/TimingProfiler/ViewModels/StatsNode.cpp
@@ -97,107 +97,151 @@ struct FStatsNodeMetaGroupName
 {
 	static const FName Time;
 	static const FName Memory;
+	static const FName Bandwidth;
+	static const FName Percent;
 };
 
 const FName FStatsNodeMetaGroupName::Time(TEXT("Time"));
 const FName FStatsNodeMetaGroupName::Memory(TEXT("Memory"));
+const FName FStatsNodeMetaGroupName::Bandwidth(TEXT("Bandwidth"));
+const FName FStatsNodeMetaGroupName::Percent(TEXT("Percent"));
+static const FText NAText = LOCTEXT("AggregatedStatsNA", "N/A");
 
 ////////////////////////////////////////////////////////////////////////////////////////////////////
 
 const FText FStatsNode::FormatAggregatedStatsValue(double ValueDbl, int64 ValueInt, bool bForTooltip) const
 {
-	if (AggregatedStats.Count > 0)
+	if (AggregatedStats.Count == 0)
 	{
-		if (GetDataType() == EStatsNodeDataType::Double)
+		return NAText;
+	}
+
+	if (GetDataType() == EStatsNodeDataType::Double)
+	{
+		//TODO: if (GetDisplayHint() == FStatsNodeDisplayHint::Seconds)
+		if (GetMetaGroupName() == FStatsNodeMetaGroupName::Time)
 		{
-			//TODO: if (GetDisplayHint() == FStatsNodeDisplayHint::Seconds)
-			if (GetMetaGroupName() == FStatsNodeMetaGroupName::Time)
+			if (bForTooltip)
 			{
-				if (bForTooltip)
-				{
-					return FText::FromString(FormatTimeAuto(ValueDbl, 2));
-				}
-				else
-				{
-					return FText::FromString(FormatTimeAuto(ValueDbl, 1));
-				}
+				return FText::FromString(FormatTimeAuto(ValueDbl, 2));
+			}
+			else
+			{
+				return FText::FromString(FormatTimeAuto(ValueDbl, 1));
+			}
+		}
+		else if (GetMetaGroupName() == FStatsNodeMetaGroupName::Percent)
+		{
+			FNumberFormattingOptions FormattingOptions;
+			FormattingOptions.MaximumFractionalDigits = 2;
+			return FText::Format(LOCTEXT("Counter_FractionValueFmt", "{0}"), FText::AsPercent(ValueDbl, &FormattingOptions));
+		}
+		else if (GetMetaGroupName() == FStatsNodeMetaGroupName::Bandwidth)
+		{
+			FNumberFormattingOptions FormattingOptions;
+			FormattingOptions.MaximumFractionalDigits = 2;
+			if (ValueInt < 1024)
+			{
+				return FText::Format(LOCTEXT("Counter_BwValueFmt1", "{0} B/s"), FText::AsNumber(ValueDbl, &FormattingOptions));
 			}
 			else
 			{
-				FNumberFormattingOptions FormattingOptions;
-				FormattingOptions.MaximumFractionalDigits = bForTooltip ? 12 : 6;
-				return FText::AsNumber(ValueDbl, &FormattingOptions);
+				// @note AsMemory only support integer values
+				return FText::Format(LOCTEXT("Counter_BwValueFmt2", "{0}/s"), FText::AsMemory(ValueInt, &FormattingOptions));
 			}
 		}
 		else
-		if (GetDataType() == EStatsNodeDataType::Int64)
 		{
-			//TODO: if (GetDisplayHint() == FStatsNodeDisplayHint::Bytes)
-			if (GetMetaGroupName() == FStatsNodeMetaGroupName::Memory)
+			FNumberFormattingOptions FormattingOptions;
+			FormattingOptions.MaximumFractionalDigits = bForTooltip ? 12 : 6;
+			return FText::AsNumber(ValueDbl, &FormattingOptions);
+		}
+	} // EStatsNodeDataType::Double
+	else if (GetDataType() == EStatsNodeDataType::Int64)
+	{
+		//TODO: if (GetDisplayHint() == FStatsNodeDisplayHint::Bytes)
+		if (GetMetaGroupName() == FStatsNodeMetaGroupName::Memory)
+		{
+			if (ValueInt > 0)
 			{
-				if (ValueInt > 0)
+				if (bForTooltip)
 				{
-					if (bForTooltip)
+					if (ValueInt < 1024)
 					{
-						if (ValueInt < 1024)
-						{
-							return FText::Format(LOCTEXT("Counter_MemValueFmt1", "{0} bytes"), FText::AsNumber(ValueInt));
-						}
-						else
-						{
-							FNumberFormattingOptions FormattingOptions;
-							FormattingOptions.MaximumFractionalDigits = 2;
-							return FText::Format(LOCTEXT("Counter_MemValueFmt2", "{0} ({1} bytes)"), FText::AsMemory(ValueInt, &FormattingOptions), FText::AsNumber(ValueInt));
-						}
+						return FText::Format(LOCTEXT("Counter_MemValueFmt1", "{0} bytes"), FText::AsNumber(ValueInt));
 					}
 					else
 					{
 						FNumberFormattingOptions FormattingOptions;
-						FormattingOptions.MaximumFractionalDigits = 1;
-						return FText::AsMemory(ValueInt, &FormattingOptions);
+						FormattingOptions.MaximumFractionalDigits = 2;
+						return FText::Format(LOCTEXT("Counter_MemValueFmt2", "{0} ({1} bytes)"), FText::AsMemory(ValueInt, &FormattingOptions), FText::AsNumber(ValueInt));
 					}
 				}
-				else if (ValueInt == 0)
+				else
 				{
-					return FText::FromString(TEXT("0"));
+					FNumberFormattingOptions FormattingOptions;
+					FormattingOptions.MaximumFractionalDigits = 1;
+					return FText::AsMemory(ValueInt, &FormattingOptions);
 				}
-				else
+			}
+			else if (ValueInt == 0)
+			{
+				return FText::FromString(TEXT("0"));
+			}
+			else
+			{
+				if (bForTooltip)
 				{
-					if (bForTooltip)
+					if (-ValueInt < 1024)
 					{
-						if (-ValueInt < 1024)
-						{
-							return FText::Format(LOCTEXT("Counter_NegMemValueFmt1", "-{0} bytes"), FText::AsNumber(-ValueInt));
-						}
-						else
-						{
-							FNumberFormattingOptions FormattingOptions;
-							FormattingOptions.MaximumFractionalDigits = 2;
-							return FText::Format(LOCTEXT("Counter_NegMemValueFmt2", "-{0} (-{1} bytes)"), FText::AsMemory(-ValueInt, &FormattingOptions), FText::AsNumber(-ValueInt));
-						}
+						return FText::Format(LOCTEXT("Counter_NegMemValueFmt1", "-{0} bytes"), FText::AsNumber(-ValueInt));
 					}
 					else
 					{
 						FNumberFormattingOptions FormattingOptions;
-						FormattingOptions.MaximumFractionalDigits = 1;
-						return FText::Format(LOCTEXT("Counter_NegMemValueFmt3", "-{0}"), FText::AsMemory(-ValueInt, &FormattingOptions));
+						FormattingOptions.MaximumFractionalDigits = 2;
+						return FText::Format(LOCTEXT("Counter_NegMemValueFmt2", "-{0} (-{1} bytes)"), FText::AsMemory(-ValueInt, &FormattingOptions), FText::AsNumber(-ValueInt));
 					}
 				}
+				else
+				{
+					FNumberFormattingOptions FormattingOptions;
+					FormattingOptions.MaximumFractionalDigits = 1;
+					return FText::Format(LOCTEXT("Counter_NegMemValueFmt3", "-{0}"), FText::AsMemory(-ValueInt, &FormattingOptions));
+				}
+			}
+		}
+		else if (GetMetaGroupName() == FStatsNodeMetaGroupName::Bandwidth)
+		{
+			FNumberFormattingOptions FormattingOptions;
+			FormattingOptions.MaximumFractionalDigits = 2;
+			if (ValueInt < 1024)
+			{
+				return FText::Format(LOCTEXT("Counter_BwValueFmt1", "{0} B/s"), FText::AsNumber(ValueInt));
 			}
 			else
 			{
-				return FText::AsNumber(ValueInt);
+				return FText::Format(LOCTEXT("Counter_BwValueFmt2", "{0}/s"), FText::AsMemory(ValueInt, &FormattingOptions));
 			}
 		}
-	}
+		else
+		{
+			return FText::AsNumber(ValueInt);
+		}
+	} // EStatsNodeDataType::Int64
 
-	return LOCTEXT("AggregatedStatsNA", "N/A");
+	return NAText;
 }
 
 ////////////////////////////////////////////////////////////////////////////////////////////////////
 
 const FText FStatsNode::GetTextForAggregatedStatsSum(bool bForTooltip) const
 {
+	// It doesn't make sense to print sum values for bandwidth or fractional values
+	if (GetMetaGroupName() == FStatsNodeMetaGroupName::Bandwidth || GetMetaGroupName() == FStatsNodeMetaGroupName::Percent)
+	{
+		return NAText;
+	}
 	return FormatAggregatedStatsValue(AggregatedStats.DoubleStats.Sum, AggregatedStats.Int64Stats.Sum, bForTooltip);
 }
 
diff --git a/Engine/Source/Developer/TraceInsights/Private/Insights/TimingProfiler/Widgets/SStatsView.cpp b/Engine/Source/Developer/TraceInsights/Private/Insights/TimingProfiler/Widgets/SStatsView.cpp
--- a/Engine/Source/Developer/TraceInsights/Private/Insights/TimingProfiler/Widgets/SStatsView.cpp
+++ b/Engine/Source/Developer/TraceInsights/Private/Insights/TimingProfiler/Widgets/SStatsView.cpp
@@ -2014,21 +2014,36 @@ void SStatsView::RebuildTree(bool bResync)
 			StatsNodes.Reserve(CounterCount);
 			StatsNodesIdMap.Reserve(CounterCount);
 
-			const FName MemoryGroup(TEXT("Memory"));
-			const FName MiscFloatGroup(TEXT("Misc_double"));
-			const FName MiscInt64Group(TEXT("Misc_int64"));
 
 			// Add nodes only for new counters.
 			uint32 CounterIndex = PreviousNodeCount;
-			CountersProvider.EnumerateCounters([this, &CounterIndex, MemoryGroup, MiscFloatGroup, MiscInt64Group](uint32 CounterId, const TraceServices::ICounter& Counter)
+			CountersProvider.EnumerateCounters([this, &CounterIndex](uint32 CounterId, const TraceServices::ICounter& Counter)
 			{
+				auto DeduceGroup = [](const TraceServices::ICounter& Counter)
+				{
+					static const FName MemoryGroup(TEXT("Memory"));
+					static const FName BandwidthGroup(TEXT("Bandwidth"));
+					static const FName PercentGroup(TEXT("Percent"));
+					static const FName MiscFloatGroup(TEXT("Misc_double"));
+					static const FName MiscInt64Group(TEXT("Misc_int64"));
+
+					switch (Counter.GetDisplayHint())	
+					{
+					case TraceServices::CounterDisplayHint_Memory:
+						return MemoryGroup;
+					case TraceServices::CounterDisplayHint_Bandwidth:
+						return BandwidthGroup;
+					case TraceServices::CounterDisplayHint_Percent:
+						return PercentGroup;
+					}
+					return Counter.IsFloatingPoint() ? MiscFloatGroup : MiscInt64Group;
+				};
+
 				FStatsNodePtr NodePtr = StatsNodesIdMap.FindRef(CounterId);
 				if (!NodePtr)
 				{
 					FName Name(Counter.GetName());
-					const FName Group = Counter.GetGroup() ? Counter.GetGroup() :
-										((Counter.GetDisplayHint() == TraceServices::CounterDisplayHint_Memory) ? MemoryGroup :
-										  Counter.IsFloatingPoint() ? MiscFloatGroup : MiscInt64Group);
+					const FName Group = Counter.GetGroup() ? Counter.GetGroup() : DeduceGroup(Counter);
 					const EStatsNodeType Type = EStatsNodeType::Counter;
 					const EStatsNodeDataType DataType = Counter.IsFloatingPoint() ? EStatsNodeDataType::Double : EStatsNodeDataType::Int64;
 					NodePtr = MakeShared<FStatsNode>(CounterId, Name, Group, Type, DataType);
diff --git a/Engine/Source/Developer/TraceServices/Public/TraceServices/Model/Counters.h b/Engine/Source/Developer/TraceServices/Public/TraceServices/Model/Counters.h
--- a/Engine/Source/Developer/TraceServices/Public/TraceServices/Model/Counters.h
+++ b/Engine/Source/Developer/TraceServices/Public/TraceServices/Model/Counters.h
@@ -14,6 +14,8 @@ enum ECounterDisplayHint
 {
 	CounterDisplayHint_None,
 	CounterDisplayHint_Memory,
+	CounterDisplayHint_Bandwidth,
+	CounterDisplayHint_Percent,
 };
 
 enum class ECounterOpType : uint8
