Add compact tab spawner menu in Verse Insights toolbar.
Change the Verse Insights toolbar to use ToolMenus to enable filtering of UI elements.

diff --git a/Engine/Source/Developer/TraceInsights/Private/Insights/InsightsStyle.cpp b/Engine/Source/Developer/TraceInsights/Private/Insights/InsightsStyle.cpp
--- a/Engine/Source/Developer/TraceInsights/Private/Insights/InsightsStyle.cpp
+++ b/Engine/Source/Developer/TraceInsights/Private/Insights/InsightsStyle.cpp
@@ -153,6 +153,8 @@ void FInsightsStyle::FStyle::Initialize()
 
 	Set("Icons.TasksView", new IMAGE_BRUSH_SVG("Tasks", Icon16x16));
 	Set("Icons.PackagesView", new EDITOR_IMAGE_BRUSH_SVG("Starship/Common/Spreadsheet", Icon16x16));
+	
+	Set("Icons.PanelMenu.Toolbar", new CORE_IMAGE_BRUSH_SVG("Starship/Common/Menu", Icon20x20));
 
 	//////////////////////////////////////////////////
 	// Timing View
diff --git a/Engine/Source/Developer/TraceInsights/Private/Insights/TimingProfiler/Widgets/STimingProfilerToolbar.cpp b/Engine/Source/Developer/TraceInsights/Private/Insights/TimingProfiler/Widgets/STimingProfilerToolbar.cpp
--- a/Engine/Source/Developer/TraceInsights/Private/Insights/TimingProfiler/Widgets/STimingProfilerToolbar.cpp
+++ b/Engine/Source/Developer/TraceInsights/Private/Insights/TimingProfiler/Widgets/STimingProfilerToolbar.cpp
@@ -4,12 +4,16 @@
 
 #include "Framework/MultiBox/MultiBoxExtender.h"
 #include "Framework/MultiBox/MultiBoxBuilder.h"
+#include "Modules/ModuleManager.h"
+#include "ToolMenus.h"
+#include "ToolMenuDelegates.h"
 #include "Widgets/SBoxPanel.h"
 
 // TraceInsights
 #include "Insights/InsightsCommands.h"
 #include "Insights/InsightsManager.h"
 #include "Insights/InsightsStyle.h"
+#include "Insights/IUnrealInsightsModule.h"
 #include "Insights/TimingProfiler/TimingProfilerCommands.h"
 
 ////////////////////////////////////////////////////////////////////////////////////////////////////
@@ -37,102 +41,133 @@ void STimingProfilerToolbar::Construct(const FArguments& InArgs, const FInsights
 {
 	struct Local
 	{
-		static void FillViewToolbar(FToolBarBuilder& ToolbarBuilder, const FArguments &InArgs, const FInsightsMajorTabConfig& Config)
+		static void MakeTabSpawnerMenu(FToolMenuSection& Section, const FInsightsMajorTabConfig& Config)
 		{
-			if (InArgs._ToolbarExtender.IsValid())
+			if (Config.ShouldRegisterMinorTab(FTimingProfilerTabs::FramesTrackID))
 			{
-				InArgs._ToolbarExtender->Apply("BeforeMainToolbar", EExtensionHook::First, ToolbarBuilder);
+				Section.AddMenuEntry(FTimingProfilerCommands::Get().ToggleFramesTrackVisibility, TAttribute<FText>(), TAttribute<FText>(), FSlateIcon(FInsightsStyle::GetStyleSetName(), "Icons.FramesTrack.ToolBar"));
 			}
+			if (Config.ShouldRegisterMinorTab(FTimingProfilerTabs::TimingViewID))
+			{
+				Section.AddMenuEntry(FTimingProfilerCommands::Get().ToggleTimingViewVisibility, TAttribute<FText>(), TAttribute<FText>(), FSlateIcon(FInsightsStyle::GetStyleSetName(), "Icons.TimingView.ToolBar"));
+			}
+			if (Config.ShouldRegisterMinorTab(FTimingProfilerTabs::TimersID))
+			{
+				Section.AddMenuEntry(FTimingProfilerCommands::Get().ToggleTimersViewVisibility, TAttribute<FText>(), TAttribute<FText>(), FSlateIcon(FInsightsStyle::GetStyleSetName(), "Icons.TimersView.ToolBar"));
+			}
+			if (Config.ShouldRegisterMinorTab(FTimingProfilerTabs::CallersID))
+			{
+				Section.AddMenuEntry(FTimingProfilerCommands::Get().ToggleCallersTreeViewVisibility, TAttribute<FText>(), TAttribute<FText>(), FSlateIcon(FInsightsStyle::GetStyleSetName(), "Icons.CallersView.ToolBar"));
+			}
+			if (Config.ShouldRegisterMinorTab(FTimingProfilerTabs::CalleesID))
+			{
+				Section.AddMenuEntry(FTimingProfilerCommands::Get().ToggleCalleesTreeViewVisibility, TAttribute<FText>(), TAttribute<FText>(), FSlateIcon(FInsightsStyle::GetStyleSetName(), "Icons.CalleesView.ToolBar"));
+			}
+			if (Config.ShouldRegisterMinorTab(FTimingProfilerTabs::StatsCountersID))
+			{
+				Section.AddMenuEntry(FTimingProfilerCommands::Get().ToggleStatsCountersViewVisibility, TAttribute<FText>(), TAttribute<FText>(), FSlateIcon(FInsightsStyle::GetStyleSetName(), "Icons.CountersView.ToolBar"));
+			}
+			if (Config.ShouldRegisterMinorTab(FTimingProfilerTabs::LogViewID))
+			{
+				Section.AddMenuEntry(FTimingProfilerCommands::Get().ToggleLogViewVisibility, TAttribute<FText>(), TAttribute<FText>(), FSlateIcon(FInsightsStyle::GetStyleSetName(), "Icons.LogView.ToolBar"));
+			}
+			if (Config.ShouldRegisterMinorTab(FTimingProfilerTabs::ModulesViewID))
+			{
+				Section.AddMenuEntry(FTimingProfilerCommands::Get().ToggleModulesViewVisibility, TAttribute<FText>(), TAttribute<FText>(), FSlateIcon(FInsightsStyle::GetStyleSetName(), "Icons.ModulesView.ToolBar"));
+			}
+		}
 
-			ToolbarBuilder.BeginSection("View");
+		static void FillViewToolbar(UToolMenu* MainToolBar, const FArguments& InArgs, const FInsightsMajorTabConfig& Config)
+		{
+			if (Config.bUseCompactTabSpawners)
 			{
-				if (Config.ShouldRegisterMinorTab(FTimingProfilerTabs::FramesTrackID))
-				{
-					ToolbarBuilder.AddToolBarButton(FTimingProfilerCommands::Get().ToggleFramesTrackVisibility,
-						NAME_None, TAttribute<FText>(), TAttribute<FText>(),
-						FSlateIcon(FInsightsStyle::GetStyleSetName(), "Icons.FramesTrack.ToolBar"));
-				}
+				const FName MenuName = FName("TimingProfiler.MainToolBar.PanelSelectorMenu");
+				const UToolMenus* ToolMenus = UToolMenus::Get();
 
-				if (Config.ShouldRegisterMinorTab(FTimingProfilerTabs::TimingViewID))
+				if (ToolMenus->IsMenuRegistered(MenuName))
 				{
-					ToolbarBuilder.AddToolBarButton(FTimingProfilerCommands::Get().ToggleTimingViewVisibility,
-						NAME_None, TAttribute<FText>(), TAttribute<FText>(),
-						FSlateIcon(FInsightsStyle::GetStyleSetName(), "Icons.TimingView.ToolBar"));
+					return;
 				}
 
-				if (Config.ShouldRegisterMinorTab(FTimingProfilerTabs::TimersID))
+				UToolMenu* PanelsMenu = UToolMenus::Get()->RegisterMenu(MenuName, NAME_None, EMultiBoxType::Menu);
+				if (PanelsMenu == nullptr)
 				{
-					ToolbarBuilder.AddToolBarButton(FTimingProfilerCommands::Get().ToggleTimersViewVisibility,
-						NAME_None, TAttribute<FText>(), TAttribute<FText>(),
-						FSlateIcon(FInsightsStyle::GetStyleSetName(), "Icons.TimersView.ToolBar"));
+					return;
 				}
-
-				if (Config.ShouldRegisterMinorTab(FTimingProfilerTabs::CallersID))
+				FToolMenuSection& Section = PanelsMenu->AddSection("InsightsTabs");
+				Local::MakeTabSpawnerMenu(Section, Config);
+
+				FToolMenuSection& ToolbarSection = MainToolBar->AddSection(FName("Main"));
+				ToolbarSection.AddEntry(FToolMenuEntry::InitComboButton(
+					"PanelSelector",
+					FUIAction(),
+					FNewToolMenuChoice(FNewToolMenuWidget::CreateLambda([MenuName](const FToolMenuContext& Ctx) {
+						return UToolMenus::Get()->GenerateWidget(MenuName, Ctx);
+						})),
+					LOCTEXT("PanelsText", "Panels"),
+					LOCTEXT("PanelsToolTip", "Open/Close the available panels."),
+					FSlateIcon(FInsightsStyle::GetStyleSetName(), "Icons.PanelMenu.Toolbar"),
+					false));
+
+				FSlimHorizontalToolBarBuilder ToolbarBuilder(nullptr, FMultiBoxCustomization::None);
+				ToolbarBuilder.SetStyle(&FInsightsStyle::Get(), "LegacyExtendersToolbarStyle");
+
+				if (InArgs._ToolbarExtender.IsValid())
 				{
-					ToolbarBuilder.AddToolBarButton(FTimingProfilerCommands::Get().ToggleCallersTreeViewVisibility,
-						NAME_None, TAttribute<FText>(), TAttribute<FText>(),
-						FSlateIcon(FInsightsStyle::GetStyleSetName(), "Icons.CallersView.ToolBar"));
+					InArgs._ToolbarExtender->Apply("MainToolbar", EExtensionHook::First, ToolbarBuilder);
 				}
 
-				if (Config.ShouldRegisterMinorTab(FTimingProfilerTabs::CalleesID))
-				{
-					ToolbarBuilder.AddToolBarButton(FTimingProfilerCommands::Get().ToggleCalleesTreeViewVisibility,
-						NAME_None, TAttribute<FText>(), TAttribute<FText>(),
-						FSlateIcon(FInsightsStyle::GetStyleSetName(), "Icons.CalleesView.ToolBar"));
-				}
-				if (Config.ShouldRegisterMinorTab(FTimingProfilerTabs::StatsCountersID))
-				{
-					ToolbarBuilder.AddToolBarButton(FTimingProfilerCommands::Get().ToggleStatsCountersViewVisibility,
-						NAME_None, TAttribute<FText>(), TAttribute<FText>(),
-						FSlateIcon(FInsightsStyle::GetStyleSetName(), "Icons.CountersView.ToolBar"));
-				}
-				if (Config.ShouldRegisterMinorTab(FTimingProfilerTabs::LogViewID))
-				{
-					ToolbarBuilder.AddToolBarButton(FTimingProfilerCommands::Get().ToggleLogViewVisibility,
-						NAME_None, TAttribute<FText>(), TAttribute<FText>(),
-						FSlateIcon(FInsightsStyle::GetStyleSetName(), "Icons.LogView.ToolBar"));
-				}
-				if (Config.ShouldRegisterMinorTab(FTimingProfilerTabs::ModulesViewID))
+				ToolbarSection.AddEntry(FToolMenuEntry::InitWidget(
+					"LegacyExtenders",
+					ToolbarBuilder.MakeWidget(),
+					FText()));
+			}
+			else
+			{
+				FToolMenuSection& ToolbarSection = MainToolBar->AddSection(FName("Main"));
+				Local::MakeTabSpawnerMenu(ToolbarSection, Config);
+
+				FSlimHorizontalToolBarBuilder ToolbarBuilder(nullptr, FMultiBoxCustomization::None);
+				ToolbarBuilder.SetStyle(&FInsightsStyle::Get(), "LegacyExtendersToolbarStyle");
+
+				if (InArgs._ToolbarExtender.IsValid())
 				{
-					ToolbarBuilder.AddToolBarButton(FTimingProfilerCommands::Get().ToggleModulesViewVisibility,
-						NAME_None, TAttribute<FText>(), TAttribute<FText>(),
-						FSlateIcon(FInsightsStyle::GetStyleSetName(), "Icons.ModulesView.ToolBar"));
+					InArgs._ToolbarExtender->Apply("MainToolbar", EExtensionHook::First, ToolbarBuilder);
 				}
-			}
-			ToolbarBuilder.EndSection();
 
-			if (InArgs._ToolbarExtender.IsValid())
-			{
-				InArgs._ToolbarExtender->Apply("MainToolbar", EExtensionHook::First, ToolbarBuilder);
+				ToolbarSection.AddEntry(FToolMenuEntry::InitWidget(
+					"LegacyExtenders",
+					ToolbarBuilder.MakeWidget(),
+					FText()));
 			}
 		}
 
-		static void FillRightSideToolbar(FToolBarBuilder& ToolbarBuilder, const FArguments &InArgs)
+		static void FillRightSideToolbar(UToolMenu* ToolBar, const FArguments& InArgs, const FInsightsMajorTabConfig& Config)
 		{
-			ToolbarBuilder.BeginSection("Debug");
-			{
-				ToolbarBuilder.AddToolBarButton(FInsightsCommands::Get().ToggleDebugInfo,
-					NAME_None, FText(), TAttribute<FText>(),
-					FSlateIcon(FInsightsStyle::GetStyleSetName(), "Icons.Debug.ToolBar"));
-			}
-			ToolbarBuilder.EndSection();
+			FToolMenuSection& Section = ToolBar->AddSection("Debug");
 
-			if (InArgs._ToolbarExtender.IsValid())
-			{
-				InArgs._ToolbarExtender->Apply("RightSideToolbar", EExtensionHook::First, ToolbarBuilder);
-			}
+			Section.AddEntry(FToolMenuEntry::InitToolBarButton(FInsightsCommands::Get().ToggleDebugInfo,
+				FText(), FText(), FSlateIcon(FInsightsStyle::GetStyleSetName(), "Icons.Debug.ToolBar")));
 		}
 	};
 
 	TSharedPtr<FUICommandList> CommandList = FInsightsManager::Get()->GetCommandList();
 
-	FSlimHorizontalToolBarBuilder ToolbarBuilder(CommandList.ToSharedRef(), FMultiBoxCustomization::None);
-	ToolbarBuilder.SetStyle(&FInsightsStyle::Get(), "PrimaryToolbar");
-	Local::FillViewToolbar(ToolbarBuilder, InArgs, Config);
+	const FName ToolBarName = FName("TimingProfiler.MainToolBar");
+	const FName RightSideToolBarName = FName("TimingProfiler.RightToolBar");
+	const UToolMenus* ToolMenus = UToolMenus::Get();
+		
+	if (!ToolMenus->IsMenuRegistered(ToolBarName))
+	{
+		FToolMenuOwnerScoped OwnerScoped(UE_MODULE_NAME);
+		UToolMenu* MainToolBar = UToolMenus::Get()->RegisterMenu(ToolBarName, NAME_None, EMultiBoxType::SlimHorizontalToolBar);
+		Local::FillViewToolbar(MainToolBar, InArgs, Config);
+
+		UToolMenu* RightSideToolBar = UToolMenus::Get()->RegisterMenu(RightSideToolBarName, NAME_None, EMultiBoxType::SlimHorizontalToolBar);
+		Local::FillRightSideToolbar(RightSideToolBar, InArgs, Config);
+	}
 
-	FSlimHorizontalToolBarBuilder RightSideToolbarBuilder(CommandList.ToSharedRef(), FMultiBoxCustomization::None);
-	RightSideToolbarBuilder.SetStyle(&FInsightsStyle::Get(), "PrimaryToolbar");
-	Local::FillRightSideToolbar(RightSideToolbarBuilder, InArgs);
+	FToolMenuContext MenuContext(CommandList, InArgs._ToolbarExtender);
 
 	ChildSlot
 	[
@@ -144,7 +179,7 @@ void STimingProfilerToolbar::Construct(const FArguments& InArgs, const FInsights
 		.FillWidth(1.0)
 		.Padding(0.0f)
 		[
-			ToolbarBuilder.MakeWidget()
+			UToolMenus::Get()->GenerateWidget(ToolBarName, MenuContext)
 		]
 
 		+ SHorizontalBox::Slot()
@@ -153,7 +188,7 @@ void STimingProfilerToolbar::Construct(const FArguments& InArgs, const FInsights
 		.AutoWidth()
 		.Padding(0.0f)
 		[
-			RightSideToolbarBuilder.MakeWidget()
+			UToolMenus::Get()->GenerateWidget(RightSideToolBarName, MenuContext)
 		]
 	];
 }
diff --git a/Engine/Source/Developer/TraceInsights/Public/Insights/IUnrealInsightsModule.h b/Engine/Source/Developer/TraceInsights/Public/Insights/IUnrealInsightsModule.h
--- a/Engine/Source/Developer/TraceInsights/Public/Insights/IUnrealInsightsModule.h
+++ b/Engine/Source/Developer/TraceInsights/Public/Insights/IUnrealInsightsModule.h
@@ -145,6 +145,9 @@ struct FInsightsMajorTabConfig
 	/** A denylist that can be used to suppress default minor tabs. */
 	TSet<FName> DisabledMinorTabIds;
 
+	/** If true, show the tab spawner entries in a submenu. Otherwise show them as individual buttons in the toolbar. */
+	bool bUseCompactTabSpawners = false;
+
 	/** Whether the tab is available for selection (i.e. registered with the tab manager) */
 	bool bIsAvailable;
 };
diff --git a/Engine/Source/Developer/TraceInsights/TraceInsights.Build.cs b/Engine/Source/Developer/TraceInsights/TraceInsights.Build.cs
--- a/Engine/Source/Developer/TraceInsights/TraceInsights.Build.cs
+++ b/Engine/Source/Developer/TraceInsights/TraceInsights.Build.cs
@@ -60,7 +60,8 @@ public class TraceInsights : ModuleRules
 			new string[] {
 				"ImageCore",
 				"SlateCore",
-				"ToolWidgets"
+				"ToolWidgets",
+				"ToolMenus",
 			}
 		);
 
diff --git a/Engine/Source/Developer/TraceInsightsCore/Private/InsightsCore/Common/InsightsCoreStyle.cpp b/Engine/Source/Developer/TraceInsightsCore/Private/InsightsCore/Common/InsightsCoreStyle.cpp
--- a/Engine/Source/Developer/TraceInsightsCore/Private/InsightsCore/Common/InsightsCoreStyle.cpp
+++ b/Engine/Source/Developer/TraceInsightsCore/Private/InsightsCore/Common/InsightsCoreStyle.cpp
@@ -265,6 +265,13 @@ void FInsightsCoreStyle::FStyle::Initialize()
 		Set("PrimaryToolbar.MaxUniformToolbarSize", 40.0f);
 	}
 
+	// LegacyExtendersToolbarStyle
+	{
+		FToolBarStyle LegacyExtendersToolbarStyle = FAppStyle::Get().GetWidgetStyle<FToolBarStyle>("SlimToolBar");
+		LegacyExtendersToolbarStyle.BackgroundPadding = FMargin(0.0f);
+		Set("LegacyExtendersToolbarStyle", LegacyExtendersToolbarStyle);
+	}
+
 	// SecondaryToolbar
 	{
 		FToolBarStyle SecondaryToolbarStyle = FAppStyle::Get().GetWidgetStyle<FToolBarStyle>("SlimToolBar");
