[Insights] Added handling for invalid timer nodes in some operations related to graph series in the Timers panel and in the Callers & Callees panels. Also made few minor code changes to minimize code differences between the two implementations.

#rnx

diff --git a/Engine/Source/Developer/TraceInsights/Private/Insights/TimingProfiler/Widgets/STimerTreeView.cpp b/Engine/Source/Developer/TraceInsights/Private/Insights/TimingProfiler/Widgets/STimerTreeView.cpp
--- a/Engine/Source/Developer/TraceInsights/Private/Insights/TimingProfiler/Widgets/STimerTreeView.cpp
+++ b/Engine/Source/Developer/TraceInsights/Private/Insights/TimingProfiler/Widgets/STimerTreeView.cpp
@@ -1758,36 +1758,50 @@ void STimerTreeView::ToggleGraphFrameStatsSeries(TSharedRef<FTimingGraphTrack> G
 
 bool STimerTreeView::IsFrameStatsSeriesInTimingViewMainGraph(FTimerNodePtr TimerNode, ETraceFrameType FrameType) const
 {
-	TSharedPtr<FTimingGraphTrack> GraphTrack = GetTimingViewMainGraphTrack();
-
-	if (GraphTrack.IsValid())
+	if (!TimerNode.IsValid())
 	{
-		const uint32 TimerId = TimerNode->GetTimerId();
-		TSharedPtr<FTimingGraphSeries> Series = GraphTrack->GetFrameStatsTimerSeries(TimerId, FrameType);
+		return false;
+	}
 
-		return Series.IsValid();
+	TSharedPtr<FTimingGraphTrack> GraphTrack = GetTimingViewMainGraphTrack();
+	if (!GraphTrack.IsValid())
+	{
+		return false;
 	}
 
-	return false;
+	const uint32 TimerId = TimerNode->GetTimerId();
+	TSharedPtr<FTimingGraphSeries> Series = GraphTrack->GetFrameStatsTimerSeries(TimerId, FrameType);
+	return Series.IsValid();
 }
 
 ////////////////////////////////////////////////////////////////////////////////////////////////////
 
 void STimerTreeView::ToggleTimingViewMainGraphEventFrameStatsSeries(FTimerNodePtr TimerNode, ETraceFrameType FrameType) const
 {
+	if (!TimerNode.IsValid())
+	{
+		return;
+	}
+
 	TSharedPtr<FTimingGraphTrack> GraphTrack = GetTimingViewMainGraphTrack();
-	if (GraphTrack.IsValid())
+	if (!GraphTrack.IsValid())
 	{
-		ToggleGraphFrameStatsSeries(GraphTrack.ToSharedRef(), TimerNode.ToSharedRef(), FrameType);
+		return;
 	}
+
+	ToggleGraphFrameStatsSeries(GraphTrack.ToSharedRef(), TimerNode.ToSharedRef(), FrameType);
 }
 
 ////////////////////////////////////////////////////////////////////////////////////////////////////
 
 bool STimerTreeView::IsSeriesInFrameTrack(FTimerNodePtr TimerNode, ETraceFrameType FrameType) const
 {
-	TSharedPtr<SFrameTrack> FrameTrack = GetFrameTrack();
+	if (!TimerNode.IsValid())
+	{
+		return false;
+	}
 
+	TSharedPtr<SFrameTrack> FrameTrack = GetFrameTrack();
 	if (!FrameTrack.IsValid())
 	{
 		return false;
@@ -1800,6 +1814,11 @@ bool STimerTreeView::IsSeriesInFrameTrack(FTimerNodePtr TimerNode, ETraceFrameTy
 
 void STimerTreeView::ToggleFrameTrackSeries(FTimerNodePtr TimerNode, ETraceFrameType FrameType) const
 {
+	if (!TimerNode.IsValid())
+	{
+		return;
+	}
+
 	TSharedPtr<SFrameTrack> FrameTrack = GetFrameTrack();
 	if (!FrameTrack.IsValid())
 	{
@@ -1919,9 +1938,10 @@ void STimerTreeView::OpenSourceFileInIDE(FTimerNodePtr InNode) const
 		return;
 	}
 
+	ISourceCodeAccessModule& SourceCodeAccessModule = FModuleManager::LoadModuleChecked<ISourceCodeAccessModule>("SourceCodeAccess");
+
 	FString UsablePath;
 	const bool bPathExists = FInsightsManager::Get()->GetSourceFilePathHelper().GetUsableFilePath(File, UsablePath);
-	ISourceCodeAccessModule& SourceCodeAccessModule = FModuleManager::LoadModuleChecked<ISourceCodeAccessModule>("SourceCodeAccess");
 	if (bPathExists)
 	{
 		ISourceCodeAccessor& SourceCodeAccessor = SourceCodeAccessModule.GetAccessor();
@@ -1948,7 +1968,7 @@ bool STimerTreeView::ContextMenu_FindInstance_CanExecute() const
 void STimerTreeView::ContextMenu_FindInstance_Execute(bool bFindMax) const
 {
 	FTimerNodePtr SelectedNode = GetSingleSelectedTimerNode();
-	if (!SelectedNode.IsValid())
+	if (!SelectedNode.IsValid() || SelectedNode->GetType() == ETimerNodeType::Group)
 	{
 		return;
 	}
@@ -1974,12 +1994,12 @@ bool STimerTreeView::ContextMenu_FindInstanceInSelection_CanExecute() const
 	}
 
 	TSharedPtr<STimingView> TimingView = GetTimingView();
-	if (TimingView.IsValid())
+	if (!TimingView.IsValid())
 	{
-		return TimingView->GetSelectionEndTime() > TimingView->GetSelectionStartTime();
+		return false;
 	}
 
-	return false;
+	return TimingView->GetSelectionEndTime() > TimingView->GetSelectionStartTime();
 }
 
 ////////////////////////////////////////////////////////////////////////////////////////////////////
diff --git a/Engine/Source/Developer/TraceInsights/Private/Insights/TimingProfiler/Widgets/STimersView.cpp b/Engine/Source/Developer/TraceInsights/Private/Insights/TimingProfiler/Widgets/STimersView.cpp
--- a/Engine/Source/Developer/TraceInsights/Private/Insights/TimingProfiler/Widgets/STimersView.cpp
+++ b/Engine/Source/Developer/TraceInsights/Private/Insights/TimingProfiler/Widgets/STimersView.cpp
@@ -3135,36 +3135,50 @@ void STimersView::ToggleGraphFrameStatsSeries(TSharedRef<FTimingGraphTrack> Grap
 
 bool STimersView::IsFrameStatsSeriesInTimingViewMainGraph(FTimerNodePtr TimerNode, ETraceFrameType FrameType) const
 {
-	TSharedPtr<FTimingGraphTrack> GraphTrack = GetTimingViewMainGraphTrack();
-
-	if (GraphTrack.IsValid())
+	if (!TimerNode.IsValid())
 	{
-		const uint32 TimerId = TimerNode->GetTimerId();
-		TSharedPtr<FTimingGraphSeries> Series = GraphTrack->GetFrameStatsTimerSeries(TimerId, FrameType);
+		return false;
+	}
 
-		return Series.IsValid();
+	TSharedPtr<FTimingGraphTrack> GraphTrack = GetTimingViewMainGraphTrack();
+	if (!GraphTrack.IsValid())
+	{
+		return false;
 	}
 
-	return false;
+	const uint32 TimerId = TimerNode->GetTimerId();
+	TSharedPtr<FTimingGraphSeries> Series = GraphTrack->GetFrameStatsTimerSeries(TimerId, FrameType);
+	return Series.IsValid();
 }
 
 ////////////////////////////////////////////////////////////////////////////////////////////////////
 
 void STimersView::ToggleTimingViewMainGraphEventFrameStatsSeries(FTimerNodePtr TimerNode, ETraceFrameType FrameType) const
 {
+	if (!TimerNode.IsValid())
+	{
+		return;
+	}
+
 	TSharedPtr<FTimingGraphTrack> GraphTrack = GetTimingViewMainGraphTrack();
-	if (GraphTrack.IsValid())
+	if (!GraphTrack.IsValid())
 	{
-		ToggleGraphFrameStatsSeries(GraphTrack.ToSharedRef(), TimerNode.ToSharedRef(), FrameType);
+		return;
 	}
+
+	ToggleGraphFrameStatsSeries(GraphTrack.ToSharedRef(), TimerNode.ToSharedRef(), FrameType);
 }
 
 ////////////////////////////////////////////////////////////////////////////////////////////////////
 
 bool STimersView::IsSeriesInFrameTrack(FTimerNodePtr TimerNode, ETraceFrameType FrameType) const
 {
-	TSharedPtr<SFrameTrack> FrameTrack = GetFrameTrack();
+	if (!TimerNode.IsValid())
+	{
+		return false;
+	}
 
+	TSharedPtr<SFrameTrack> FrameTrack = GetFrameTrack();
 	if (!FrameTrack.IsValid())
 	{
 		return false;
@@ -3177,8 +3191,12 @@ bool STimersView::IsSeriesInFrameTrack(FTimerNodePtr TimerNode, ETraceFrameType
 
 void STimersView::ToggleFrameTrackSeries(FTimerNodePtr TimerNode, ETraceFrameType FrameType) const
 {
-	TSharedPtr<SFrameTrack> FrameTrack = GetFrameTrack();
+	if (!TimerNode.IsValid())
+	{
+		return;
+	}
 
+	TSharedPtr<SFrameTrack> FrameTrack = GetFrameTrack();
 	if (!FrameTrack.IsValid())
 	{
 		return;
@@ -3205,6 +3223,8 @@ void STimersView::ToggleFrameTrackSeries(FTimerNodePtr TimerNode, ETraceFrameTyp
 	}
 }
 
+////////////////////////////////////////////////////////////////////////////////////////////////////
+// Copy to Clipboard
 ////////////////////////////////////////////////////////////////////////////////////////////////////
 
 bool STimersView::ContextMenu_CopyToClipboard_CanExecute() const
@@ -3692,7 +3712,7 @@ IFileHandle* STimersView::OpenExportFile(const TCHAR* InFilename) const
 }
 
 ////////////////////////////////////////////////////////////////////////////////////////////////////
-// OpenSource
+// Open Source File in IDE
 ////////////////////////////////////////////////////////////////////////////////////////////////////
 
 bool STimersView::ContextMenu_OpenSource_CanExecute() const
@@ -3743,16 +3763,18 @@ void STimersView::OpenSourceFileInIDE(FTimerNodePtr InNode) const
 		return;
 	}
 
-	FString Path;
 	ISourceCodeAccessModule& SourceCodeAccessModule = FModuleManager::LoadModuleChecked<ISourceCodeAccessModule>("SourceCodeAccess");
-	if (FInsightsManager::Get()->GetSourceFilePathHelper().GetUsableFilePath(File, Path))
+
+	FString UsablePath;
+	const bool bPathExists = FInsightsManager::Get()->GetSourceFilePathHelper().GetUsableFilePath(File, UsablePath);
+	if (bPathExists)
 	{
 		ISourceCodeAccessor& SourceCodeAccessor = SourceCodeAccessModule.GetAccessor();
-		SourceCodeAccessor.OpenFileAtLine(Path, Line);
+		SourceCodeAccessor.OpenFileAtLine(UsablePath, Line);
 	}
 	else
 	{
-		SourceCodeAccessModule.OnOpenFileFailed().Broadcast(Path);
+		SourceCodeAccessModule.OnOpenFileFailed().Broadcast(UsablePath);
 	}
 }
 
@@ -3890,7 +3912,7 @@ void STimersView::ColorPickerCancelled(FLinearColor OriginalColor)
 }
 
 ////////////////////////////////////////////////////////////////////////////////////////////////////
-// Find Min/Max Instance
+// Find First/Previous/Next/Last/Min/Max Instance
 ////////////////////////////////////////////////////////////////////////////////////////////////////
 
 bool STimersView::ContextMenu_FindInstance_CanExecute() const
@@ -3904,7 +3926,7 @@ bool STimersView::ContextMenu_FindInstance_CanExecute() const
 void STimersView::ContextMenu_FindInstance_Execute(ESelectEventType EventType) const
 {
 	FTimerNodePtr SelectedNode = GetSingleSelectedTimerNode();
-	if (!SelectedNode.IsValid())
+	if (!SelectedNode.IsValid() || SelectedNode->GetType() == ETimerNodeType::Group)
 	{
 		return;
 	}
@@ -3929,12 +3951,12 @@ bool STimersView::ContextMenu_FindInstanceInSelection_CanExecute() const
 	}
 
 	TSharedPtr<STimingView> TimingView = GetTimingView();
-	if (TimingView.IsValid())
+	if (!TimingView.IsValid())
 	{
-		return TimingView->GetSelectionEndTime() > TimingView->GetSelectionStartTime();
+		return false;
 	}
 
-	return false;
+	return TimingView->GetSelectionEndTime() > TimingView->GetSelectionStartTime();
 }
 
 ////////////////////////////////////////////////////////////////////////////////////////////////////
