[Insights]
* Timing View: Added debug log messages for "Find First/Previous/Next/Last/Min/Max Event" operations that takes longer than 1ms.
* Timing View: Fixed "Find Min/Max Event" for timing events not ended (i.e. with infinite duration) to consider them having a duration from their start time only to the end of session.
* Timing View: Added tooltip display of actual inclusive time value for timing events with infinite duration (i.e. duration of timing event from start time to end of session).
* Timing View: Fixed the event selection with mouse to not scroll anymore the timing view (ex. previously the view was automatically scrolled when selecting an event that was close to left or right sides of the view or when event was only partially visible).

diff --git a/Engine/Source/Developer/TraceInsights/Private/Insights/TimingProfiler/Tracks/ThreadTimingTrack.cpp b/Engine/Source/Developer/TraceInsights/Private/Insights/TimingProfiler/Tracks/ThreadTimingTrack.cpp
--- a/Engine/Source/Developer/TraceInsights/Private/Insights/TimingProfiler/Tracks/ThreadTimingTrack.cpp
+++ b/Engine/Source/Developer/TraceInsights/Private/Insights/TimingProfiler/Tracks/ThreadTimingTrack.cpp
@@ -958,38 +958,71 @@ void FThreadTimingTrackImpl::InitTooltip(FTooltipDrawState& InOutTooltip, const
 #endif
 				}
 
-				const double TooltipEventDuration = TooltipEvent.GetDuration();
+				double SessionDuration = std::numeric_limits<double>::lowest();
+				TSharedPtr<const TraceServices::IAnalysisSession> Session = FInsightsManager::Get()->GetSession();
+				if (Session.IsValid())
+				{
+					TraceServices::FAnalysisSessionReadScope SessionReadScope(*Session.Get());
+					SessionDuration = Session->GetDurationSeconds();
+				}
+
+				double TooltipEventDuration = TooltipEvent.GetDuration();
+				const bool bIsInfinite = (TooltipEventDuration == std::numeric_limits<double>::infinity());
+				if (bIsInfinite)
+				{
+					TooltipEventDuration = SessionDuration - TooltipEvent.GetStartTime();
+				}
 
 				if (TooltipEvent.GetDepth() > 0 &&
 					ParentTimingEvent.IsValid() &&
-					ParentTimingEvent->GetDuration() > 0.0 &&
-					ParentTimingEvent->GetDuration() != std::numeric_limits<double>::infinity())
+					ParentTimingEvent->GetDuration() > 0.0)
 				{
-					const TraceServices::FTimingProfilerTimer* ParentTimer = TimerReader.GetTimer(ParentTimingEvent->GetTimerIndex());
-					const TCHAR* ParentTimerName = (ParentTimer != nullptr) ? ParentTimer->Name : TEXT("N/A");
-					FNumberFormattingOptions FormattingOptions;
-					FormattingOptions.MaximumFractionalDigits = 2;
-					const FString ValueStr = FString::Printf(TEXT("%s %s"), *FText::AsPercent(TooltipEventDuration / ParentTimingEvent->GetDuration(), &FormattingOptions).ToString(), ParentTimerName);
-					InOutTooltip.AddNameValueTextLine(TEXTVIEW("% of Parent:"), ValueStr);
+					double ParentTimingEventDuration = ParentTimingEvent->GetDuration();
+					if (ParentTimingEventDuration == std::numeric_limits<double>::infinity())
+					{
+						ParentTimingEventDuration = SessionDuration - ParentTimingEvent->GetStartTime();
+					}
+					if (ParentTimingEventDuration > 0.0)
+					{
+						const TraceServices::FTimingProfilerTimer* ParentTimer = TimerReader.GetTimer(ParentTimingEvent->GetTimerIndex());
+						const TCHAR* ParentTimerName = (ParentTimer != nullptr) ? ParentTimer->Name : TEXT("N/A");
+						FNumberFormattingOptions FormattingOptions;
+						FormattingOptions.MaximumFractionalDigits = 2;
+						const FString ValueStr = FString::Printf(TEXT("%s %s"), *FText::AsPercent(TooltipEventDuration / ParentTimingEventDuration, &FormattingOptions).ToString(), ParentTimerName);
+						InOutTooltip.AddNameValueTextLine(TEXTVIEW("% of Parent:"), ValueStr);
+					}
 				}
 
 				if (TooltipEvent.GetDepth() > 1 &&
 					RootTimingEvent.IsValid() &&
-					RootTimingEvent->GetDuration() > 0.0 &&
-					RootTimingEvent->GetDuration() != std::numeric_limits<double>::infinity())
+					RootTimingEvent->GetDuration() > 0.0)
 				{
-					const TraceServices::FTimingProfilerTimer* RootTimer = TimerReader.GetTimer(RootTimingEvent->GetTimerIndex());
-					const TCHAR* RootTimerName = (RootTimer != nullptr) ? RootTimer->Name : TEXT("N/A");
-					FNumberFormattingOptions FormattingOptions;
-					FormattingOptions.MaximumFractionalDigits = 2;
-					const FString ValueStr = FString::Printf(TEXT("%s %s"), *FText::AsPercent(TooltipEventDuration / RootTimingEvent->GetDuration(), &FormattingOptions).ToString(), RootTimerName);
-					InOutTooltip.AddNameValueTextLine(TEXTVIEW("% of Root:"), ValueStr);
+					double RootTimingEventDuration = RootTimingEvent->GetDuration();
+					if (RootTimingEventDuration == std::numeric_limits<double>::infinity())
+					{
+						RootTimingEventDuration = SessionDuration - RootTimingEvent->GetStartTime();
+					}
+					if (RootTimingEventDuration > 0.0)
+					{
+						const TraceServices::FTimingProfilerTimer* RootTimer = TimerReader.GetTimer(RootTimingEvent->GetTimerIndex());
+						const TCHAR* RootTimerName = (RootTimer != nullptr) ? RootTimer->Name : TEXT("N/A");
+						FNumberFormattingOptions FormattingOptions;
+						FormattingOptions.MaximumFractionalDigits = 2;
+						const FString ValueStr = FString::Printf(TEXT("%s %s"), *FText::AsPercent(TooltipEventDuration / RootTimingEventDuration, &FormattingOptions).ToString(), RootTimerName);
+						InOutTooltip.AddNameValueTextLine(TEXTVIEW("% of Root:"), ValueStr);
+					}
 				}
 
-				InOutTooltip.AddNameValueTextLine(TEXTVIEW("Inclusive Time:"), FormatTimeAuto(TooltipEventDuration, 2));
+				{
+					FString InclusiveTimeValue = FormatTimeAuto(TooltipEventDuration, 2);
+					if (bIsInfinite)
+					{
+						InclusiveTimeValue = FString::Printf(TEXT("\u221E (%s)"), *InclusiveTimeValue);
+					}
+					InOutTooltip.AddNameValueTextLine(TEXTVIEW("Inclusive Time:"), InclusiveTimeValue);
+				}
 
-				if (TooltipEventDuration > 0.0 &&
-					TooltipEventDuration != std::numeric_limits<double>::infinity())
+				if (TooltipEventDuration > 0.0 && !bIsInfinite)
 				{
 					const double ExclusiveTimePercent = TooltipEvent.GetExclusiveTime() / TooltipEventDuration;
 					FNumberFormattingOptions FormattingOptions;
@@ -1006,7 +1039,6 @@ void FThreadTimingTrackImpl::InitTooltip(FTooltipDrawState& InOutTooltip, const
 					AppendMetadataToTooltip(InOutTooltip, MetadataSpec, Metadata);
 				}
 
-				TSharedPtr<const TraceServices::IAnalysisSession> Session = FInsightsManager::Get()->GetSession();
 				if (Session.IsValid())
 				{
 					PostInitTooltip(InOutTooltip, TooltipEvent, *Session.Get(), TimerName);
@@ -1445,24 +1477,26 @@ void FThreadTimingTrackImpl::SetFilterConfigurator(TSharedPtr<FFilterConfigurato
 
 ////////////////////////////////////////////////////////////////////////////////////////////////////
 
-TSharedPtr<const ITimingEvent> FThreadTimingTrackImpl::FindMaxEventInstance(uint32 InTimerId, double InStartTime, double InEndTime) const
+TSharedPtr<const ITimingEvent> FThreadTimingTrackImpl::FindMaxEventInstance(uint32 InTimerId, double InStartTime, double InEndTime, double InMaxValidTime) const
 {
 	TSharedPtr<FThreadTrackEvent> TimingEvent;
 
 	ReadTimers(
-		[this, InTimerId, InStartTime, InEndTime, &TimingEvent]
+		[this, InTimerId, InStartTime, InEndTime, InMaxValidTime, &TimingEvent]
 		(const TraceServices::ITimingProfilerProvider& TimingProfilerProvider, const TraceServices::ITimingProfilerTimerReader& TimerReader)
 		{
 			ReadTimeline(
-				[this, InTimerId, InStartTime, InEndTime, &TimingEvent, &TimerReader]
+				[this, InTimerId, InStartTime, InEndTime, InMaxValidTime, &TimingEvent, &TimerReader]
 				(const TraceServices::ITimingProfilerProvider::Timeline& Timeline)
 				{
 					struct FCandidateEvent
 					{
 						double StartTime = 0.0f;
-						double EndTime = -1.0f;
-						uint32 Depth = 0;
+						double EndTime = -std::numeric_limits<double>::infinity();
+						uint32 Depth = ~0u;
 						uint32 TimerIndex = 0;
+						double GetDuration() const { return EndTime - StartTime; }
+						bool IsValid() const { return Depth != ~0u; }
 					};
 					TArray<FCandidateEvent> CandidateEvents;
 
@@ -1476,25 +1510,23 @@ TSharedPtr<const ITimingEvent> FThreadTimingTrackImpl::FindMaxEventInstance(uint
 							CandidateEvents.AddDefaulted(NumTasks);
 						};
 					Params.EventRangeCallback =
-						[InTimerId, &TimerReader, &CandidateEvents]
+						[InTimerId, &TimerReader, &CandidateEvents, InMaxValidTime]
 						(double StartTime, double EndTime, uint32 Depth, const TraceServices::FTimingProfilerEvent& Event, uint32 TaskIndex)
 						{
 							const TraceServices::FTimingProfilerTimer* Timer = TimerReader.GetTimer(Event.TimerIndex);
-							if (ensure(Timer != nullptr))
+							if (ensure(Timer != nullptr) && Timer->Id == InTimerId)
 							{
-								if (Timer->Id == InTimerId)
+								if (EndTime == std::numeric_limits<double>::infinity())
 								{
-									double CandidateDuration = CandidateEvents[TaskIndex].EndTime - CandidateEvents[TaskIndex].StartTime;
-									double EventDuration = EndTime - StartTime;
-
-									if (EventDuration > CandidateDuration)
-									{
-										CandidateEvents[TaskIndex].StartTime = StartTime;
-										CandidateEvents[TaskIndex].EndTime = EndTime;
-										CandidateEvents[TaskIndex].Depth = Depth;
-										CandidateEvents[TaskIndex].TimerIndex = Event.TimerIndex;
-									}
-
+									EndTime = InMaxValidTime;
+								}
+								const double EventDuration = EndTime - StartTime;
+								if (EventDuration > CandidateEvents[TaskIndex].GetDuration())
+								{
+									CandidateEvents[TaskIndex].StartTime = StartTime;
+									CandidateEvents[TaskIndex].EndTime = EndTime;
+									CandidateEvents[TaskIndex].Depth = Depth;
+									CandidateEvents[TaskIndex].TimerIndex = Event.TimerIndex;
 								}
 							}
 							return TraceServices::EEventEnumerate::Continue;
@@ -1506,13 +1538,12 @@ TSharedPtr<const ITimingEvent> FThreadTimingTrackImpl::FindMaxEventInstance(uint
 					FCandidateEvent BestMatch;
 					for (const FCandidateEvent& Event : CandidateEvents)
 					{
-						if ((Event.EndTime - Event.StartTime) > BestMatch.EndTime - BestMatch.StartTime)
+						if (Event.GetDuration() > BestMatch.GetDuration())
 						{
 							BestMatch = Event;
 						}
 					}
-
-					if (BestMatch.EndTime > BestMatch.StartTime)
+					if (BestMatch.IsValid())
 					{
 						TimingEvent = MakeShared<FThreadTrackEvent>(SharedThis(this), BestMatch.StartTime, BestMatch.EndTime, BestMatch.Depth);
 						TimingEvent->SetTimerId(InTimerId);
@@ -1528,24 +1559,26 @@ TSharedPtr<const ITimingEvent> FThreadTimingTrackImpl::FindMaxEventInstance(uint
 
 ////////////////////////////////////////////////////////////////////////////////////////////////////
 
-TSharedPtr<const ITimingEvent> FThreadTimingTrackImpl::FindMinEventInstance(uint32 InTimerId, double InStartTime, double InEndTime) const
+TSharedPtr<const ITimingEvent> FThreadTimingTrackImpl::FindMinEventInstance(uint32 InTimerId, double InStartTime, double InEndTime, double InMaxValidTime) const
 {
 	TSharedPtr<FThreadTrackEvent> TimingEvent;
 
 	ReadTimers(
-		[this, InTimerId, InStartTime, InEndTime, &TimingEvent]
+		[this, InTimerId, InStartTime, InEndTime, InMaxValidTime, &TimingEvent]
 		(const TraceServices::ITimingProfilerProvider& TimingProfilerProvider, const TraceServices::ITimingProfilerTimerReader& TimerReader)
 		{
 			ReadTimeline(
-				[this, InTimerId, InStartTime, InEndTime, &TimingEvent, &TimerReader]
+				[this, InTimerId, InStartTime, InEndTime, InMaxValidTime, &TimingEvent, &TimerReader]
 				(const TraceServices::ITimingProfilerProvider::Timeline& Timeline)
 				{
 					struct FCandidateEvent
 					{
-						double StartTime = -std::numeric_limits<double>::infinity();
+						double StartTime = 0.0;
 						double EndTime = std::numeric_limits<double>::infinity();
-						uint32 Depth = 0;
+						uint32 Depth = ~0u;
 						uint32 TimerIndex = 0;
+						double GetDuration() const { return EndTime - StartTime; }
+						bool IsValid() const { return Depth != ~0u; }
 					};
 					TArray<FCandidateEvent> CandidateEvents;
 
@@ -1559,24 +1592,23 @@ TSharedPtr<const ITimingEvent> FThreadTimingTrackImpl::FindMinEventInstance(uint
 							CandidateEvents.AddDefaulted(NumTasks);
 						};
 					Params.EventRangeCallback =
-						[InTimerId, &TimerReader, &CandidateEvents]
+						[InTimerId, &TimerReader, &CandidateEvents, InMaxValidTime]
 						(double StartTime, double EndTime, uint32 Depth, const TraceServices::FTimingProfilerEvent& Event, uint32 TaskIndex)
 						{
 							const TraceServices::FTimingProfilerTimer* Timer = TimerReader.GetTimer(Event.TimerIndex);
-							if (ensure(Timer != nullptr))
+							if (ensure(Timer != nullptr) && Timer->Id == InTimerId)
 							{
-								if (Timer->Id == InTimerId)
+								if (EndTime == std::numeric_limits<double>::infinity())
 								{
-									double CandidateDuration = CandidateEvents[TaskIndex].EndTime - CandidateEvents[TaskIndex].StartTime;
-									double EventDuration = EndTime - StartTime;
-
-									if (EventDuration < CandidateDuration)
-									{
-										CandidateEvents[TaskIndex].StartTime = StartTime;
-										CandidateEvents[TaskIndex].EndTime = EndTime;
-										CandidateEvents[TaskIndex].Depth = Depth;
-										CandidateEvents[TaskIndex].TimerIndex = Event.TimerIndex;
-									}
+									EndTime = InMaxValidTime;
+								}
+								const double EventDuration = EndTime - StartTime;
+								if (EventDuration < CandidateEvents[TaskIndex].GetDuration())
+								{
+									CandidateEvents[TaskIndex].StartTime = StartTime;
+									CandidateEvents[TaskIndex].EndTime = EndTime;
+									CandidateEvents[TaskIndex].Depth = Depth;
+									CandidateEvents[TaskIndex].TimerIndex = Event.TimerIndex;
 								}
 							}
 							return TraceServices::EEventEnumerate::Continue;
@@ -1588,13 +1620,12 @@ TSharedPtr<const ITimingEvent> FThreadTimingTrackImpl::FindMinEventInstance(uint
 					FCandidateEvent BestMatch;
 					for (const FCandidateEvent& Event : CandidateEvents)
 					{
-						if ((Event.EndTime - Event.StartTime) < BestMatch.EndTime - BestMatch.StartTime)
+						if (Event.GetDuration() < BestMatch.GetDuration())
 						{
 							BestMatch = Event;
 						}
 					}
-
-					if (BestMatch.StartTime != -std::numeric_limits<double>::infinity())
+					if (BestMatch.IsValid())
 					{
 						TimingEvent = MakeShared<FThreadTrackEvent>(SharedThis(this), BestMatch.StartTime, BestMatch.EndTime, BestMatch.Depth);
 						TimingEvent->SetTimerId(InTimerId);
diff --git a/Engine/Source/Developer/TraceInsights/Private/Insights/TimingProfiler/Tracks/ThreadTimingTrackPrivate.h b/Engine/Source/Developer/TraceInsights/Private/Insights/TimingProfiler/Tracks/ThreadTimingTrackPrivate.h
--- a/Engine/Source/Developer/TraceInsights/Private/Insights/TimingProfiler/Tracks/ThreadTimingTrackPrivate.h
+++ b/Engine/Source/Developer/TraceInsights/Private/Insights/TimingProfiler/Tracks/ThreadTimingTrackPrivate.h
@@ -78,8 +78,8 @@ public:
 
 	//////////////////////////////////////////////////
 
-	TSharedPtr<const ITimingEvent> FindMaxEventInstance(uint32 TimerId, double StartTime, double EndTime) const;
-	TSharedPtr<const ITimingEvent> FindMinEventInstance(uint32 TimerId, double StartTime, double EndTime) const;
+	TSharedPtr<const ITimingEvent> FindMaxEventInstance(uint32 TimerId, double StartTime, double EndTime, double MaxValidTime) const;
+	TSharedPtr<const ITimingEvent> FindMinEventInstance(uint32 TimerId, double StartTime, double EndTime, double MaxValidTime) const;
 
 protected:
 	FThreadTimingSharedState& GetSharedState() const { return SharedState; }
diff --git a/Engine/Source/Developer/TraceInsights/Private/Insights/TimingProfiler/ViewModels/ThreadTimingSharedState.cpp b/Engine/Source/Developer/TraceInsights/Private/Insights/TimingProfiler/ViewModels/ThreadTimingSharedState.cpp
--- a/Engine/Source/Developer/TraceInsights/Private/Insights/TimingProfiler/ViewModels/ThreadTimingSharedState.cpp
+++ b/Engine/Source/Developer/TraceInsights/Private/Insights/TimingProfiler/ViewModels/ThreadTimingSharedState.cpp
@@ -1036,7 +1036,7 @@ void FThreadTimingSharedState::ToggleOptionState(const TCHAR* CommandName)
 
 ////////////////////////////////////////////////////////////////////////////////////////////////////
 
-TSharedPtr<const ITimingEvent> FThreadTimingSharedState::FindMaxEventInstance(uint32 TimerId, double StartTime, double EndTime)
+TSharedPtr<const ITimingEvent> FThreadTimingSharedState::FindMaxEventInstance(uint32 TimerId, double StartTime, double EndTime, double MaxValidTime)
 {
 	auto CompareAndAssignEvent = [](TSharedPtr<const ITimingEvent>& TimingEvent, TSharedPtr<const ITimingEvent>& TrackEvent)
 	{
@@ -1059,7 +1059,7 @@ TSharedPtr<const ITimingEvent> FThreadTimingSharedState::FindMaxEventInstance(ui
 		const FCpuTimingTrack& Track = *KV.Value;
 		if (Track.IsVisible())
 		{
-			TrackEvent = Track.FindMaxEventInstance(TimerId, StartTime, EndTime);
+			TrackEvent = Track.FindMaxEventInstance(TimerId, StartTime, EndTime, MaxValidTime);
 			CompareAndAssignEvent(TimingEvent, TrackEvent);
 		}
 	}
@@ -1069,20 +1069,20 @@ TSharedPtr<const ITimingEvent> FThreadTimingSharedState::FindMaxEventInstance(ui
 		const FCpuStackSampleTimingTrack& Track = *KV.Value;
 		if (Track.IsVisible())
 		{
-			TrackEvent = Track.FindMaxEventInstance(TimerId, StartTime, EndTime);
+			TrackEvent = Track.FindMaxEventInstance(TimerId, StartTime, EndTime, MaxValidTime);
 			CompareAndAssignEvent(TimingEvent, TrackEvent);
 		}
 	}
 
 	if (OldGpu1Track.IsValid() && OldGpu1Track->IsVisible())
 	{
-		TrackEvent = OldGpu1Track->FindMaxEventInstance(TimerId, StartTime, EndTime);
+		TrackEvent = OldGpu1Track->FindMaxEventInstance(TimerId, StartTime, EndTime, MaxValidTime);
 		CompareAndAssignEvent(TimingEvent, TrackEvent);
 	}
 
 	if (OldGpu2Track.IsValid() && OldGpu2Track->IsVisible())
 	{
-		TrackEvent = OldGpu2Track->FindMaxEventInstance(TimerId, StartTime, EndTime);
+		TrackEvent = OldGpu2Track->FindMaxEventInstance(TimerId, StartTime, EndTime, MaxValidTime);
 		CompareAndAssignEvent(TimingEvent, TrackEvent);
 	}
 
@@ -1091,7 +1091,7 @@ TSharedPtr<const ITimingEvent> FThreadTimingSharedState::FindMaxEventInstance(ui
 		const FGpuQueueTimingTrack& Track = *KV.Value;
 		if (Track.IsVisible())
 		{
-			TrackEvent = Track.FindMaxEventInstance(TimerId, StartTime, EndTime);
+			TrackEvent = Track.FindMaxEventInstance(TimerId, StartTime, EndTime, MaxValidTime);
 			CompareAndAssignEvent(TimingEvent, TrackEvent);
 		}
 	}
@@ -1101,7 +1101,7 @@ TSharedPtr<const ITimingEvent> FThreadTimingSharedState::FindMaxEventInstance(ui
 
 ////////////////////////////////////////////////////////////////////////////////////////////////////
 
-TSharedPtr<const ITimingEvent> FThreadTimingSharedState::FindMinEventInstance(uint32 TimerId, double StartTime, double EndTime)
+TSharedPtr<const ITimingEvent> FThreadTimingSharedState::FindMinEventInstance(uint32 TimerId, double StartTime, double EndTime, double MaxValidTime)
 {
 	auto CompareAndAssignEvent = [](TSharedPtr<const ITimingEvent>& TimingEvent, TSharedPtr<const ITimingEvent>& TrackEvent)
 	{
@@ -1124,7 +1124,7 @@ TSharedPtr<const ITimingEvent> FThreadTimingSharedState::FindMinEventInstance(ui
 		const FCpuTimingTrack& Track = *KV.Value;
 		if (Track.IsVisible())
 		{
-			TrackEvent = Track.FindMinEventInstance(TimerId, StartTime, EndTime);
+			TrackEvent = Track.FindMinEventInstance(TimerId, StartTime, EndTime, MaxValidTime);
 			CompareAndAssignEvent(TimingEvent, TrackEvent);
 		}
 	}
@@ -1134,20 +1134,20 @@ TSharedPtr<const ITimingEvent> FThreadTimingSharedState::FindMinEventInstance(ui
 		const FCpuStackSampleTimingTrack& Track = *KV.Value;
 		if (Track.IsVisible())
 		{
-			TrackEvent = Track.FindMinEventInstance(TimerId, StartTime, EndTime);
+			TrackEvent = Track.FindMinEventInstance(TimerId, StartTime, EndTime, MaxValidTime);
 			CompareAndAssignEvent(TimingEvent, TrackEvent);
 		}
 	}
 
 	if (OldGpu1Track.IsValid() && OldGpu1Track->IsVisible())
 	{
-		TrackEvent = OldGpu1Track->FindMinEventInstance(TimerId, StartTime, EndTime);
+		TrackEvent = OldGpu1Track->FindMinEventInstance(TimerId, StartTime, EndTime, MaxValidTime);
 		CompareAndAssignEvent(TimingEvent, TrackEvent);
 	}
 
 	if (OldGpu2Track.IsValid() && OldGpu2Track->IsVisible())
 	{
-		TrackEvent = OldGpu2Track->FindMinEventInstance(TimerId, StartTime, EndTime);
+		TrackEvent = OldGpu2Track->FindMinEventInstance(TimerId, StartTime, EndTime, MaxValidTime);
 		CompareAndAssignEvent(TimingEvent, TrackEvent);
 	}
 
@@ -1156,7 +1156,7 @@ TSharedPtr<const ITimingEvent> FThreadTimingSharedState::FindMinEventInstance(ui
 		const FGpuQueueTimingTrack& Track = *KV.Value;
 		if (Track.IsVisible())
 		{
-			TrackEvent = Track.FindMinEventInstance(TimerId, StartTime, EndTime);
+			TrackEvent = Track.FindMinEventInstance(TimerId, StartTime, EndTime, MaxValidTime);
 			CompareAndAssignEvent(TimingEvent, TrackEvent);
 		}
 	}
diff --git a/Engine/Source/Developer/TraceInsights/Private/Insights/TimingProfiler/ViewModels/ThreadTimingSharedState.h b/Engine/Source/Developer/TraceInsights/Private/Insights/TimingProfiler/ViewModels/ThreadTimingSharedState.h
--- a/Engine/Source/Developer/TraceInsights/Private/Insights/TimingProfiler/ViewModels/ThreadTimingSharedState.h
+++ b/Engine/Source/Developer/TraceInsights/Private/Insights/TimingProfiler/ViewModels/ThreadTimingSharedState.h
@@ -246,8 +246,8 @@ public:
 
 	void GetVisibleTimelineIndexes(TSet<uint32>& OutSet) const;
 
-	TSharedPtr<const ITimingEvent> FindMaxEventInstance(uint32 TimerId, double StartTime, double EndTime);
-	TSharedPtr<const ITimingEvent> FindMinEventInstance(uint32 TimerId, double StartTime, double EndTime);
+	TSharedPtr<const ITimingEvent> FindMaxEventInstance(uint32 TimerId, double StartTime, double EndTime, double MaxValidTime);
+	TSharedPtr<const ITimingEvent> FindMinEventInstance(uint32 TimerId, double StartTime, double EndTime, double MaxValidTime);
 
 private:
 	void CreateThreadGroupsMenu(FMenuBuilder& MenuBuilder);
diff --git a/Engine/Source/Developer/TraceInsights/Private/Insights/Widgets/STimingView.cpp b/Engine/Source/Developer/TraceInsights/Private/Insights/Widgets/STimingView.cpp
--- a/Engine/Source/Developer/TraceInsights/Private/Insights/Widgets/STimingView.cpp
+++ b/Engine/Source/Developer/TraceInsights/Private/Insights/Widgets/STimingView.cpp
@@ -96,6 +96,38 @@
 
 #define INSIGHTS_ACTIVATE_BENCHMARK 0
 
+////////////////////////////////////////////////////////////////////////////////////////////////////
+
+class FScopedLog
+{
+public:
+	FScopedLog(const TCHAR* InName, double InMinDuration = 0.001)
+		: Name(InName), MinDuration(InMinDuration)
+	{
+		Stopwatch.Start();
+	}
+
+	~FScopedLog()
+	{
+		Stopwatch.Stop();
+		const double Duration = Stopwatch.GetAccumulatedTime();
+		if (Duration > MinDuration)
+		{
+			UE_LOG(LogTimingProfiler, Log, TEXT("%s took %.3f seconds."), Name, Duration);
+		}
+	}
+
+private:
+	UE::Insights::FStopwatch Stopwatch;
+	const TCHAR* Name;
+	double MinDuration = 0.0;
+};
+
+#define INSIGHTS_SCOPED_LOG(Name) \
+	FScopedLog ScopedLog(TEXT(#Name));
+
+////////////////////////////////////////////////////////////////////////////////////////////////////
+
 const TCHAR* GetFileActivityTypeName(TraceServices::EFileActivityType Type);
 uint32 GetFileActivityTypeColor(TraceServices::EFileActivityType Type);
 
@@ -4049,7 +4081,7 @@ void STimingView::SelectHoveredTimingTrack()
 
 void STimingView::SelectHoveredTimingEvent()
 {
-	SelectTimingEvent(HoveredEvent, true, false);
+	SelectTimingEvent(HoveredEvent, false, false);
 }
 
 ////////////////////////////////////////////////////////////////////////////////////////////////////
@@ -5567,6 +5599,8 @@ void STimingView::ClearRelations()
 
 void STimingView::FindFirstEvent(bool bUseCustomFilter)
 {
+	INSIGHTS_SCOPED_LOG(FindFirstEvent);
+
 	if (SelectedEvent.IsValid())
 	{
 		SelectedEvent.Reset();
@@ -5579,6 +5613,8 @@ void STimingView::FindFirstEvent(bool bUseCustomFilter)
 
 void STimingView::FindPrevEvent(bool bUseCustomFilter)
 {
+	INSIGHTS_SCOPED_LOG(FindPrevEvent);
+
 	TSharedPtr<const ITimingEvent> BestMatchEvent;
 	double StartTime = SelectedEvent.IsValid() ? SelectedEvent->GetStartTime() : std::numeric_limits<double>::max();
 
@@ -5622,6 +5658,8 @@ void STimingView::FindPrevEvent(bool bUseCustomFilter)
 
 void STimingView::FindNextEvent(bool bUseCustomFilter)
 {
+	INSIGHTS_SCOPED_LOG(FindNextEvent);
+
 	TSharedPtr<const ITimingEvent> BestMatchEvent;
 	double StartTime = SelectedEvent.IsValid() ? SelectedEvent->GetStartTime() : std::numeric_limits<double>::lowest();
 
@@ -5665,6 +5703,8 @@ void STimingView::FindNextEvent(bool bUseCustomFilter)
 
 void STimingView::FindLastEvent(bool bUseCustomFilter)
 {
+	INSIGHTS_SCOPED_LOG(FindLastEvent);
+
 	if (SelectedEvent.IsValid())
 	{
 		SelectedEvent.Reset();
@@ -5677,26 +5717,31 @@ void STimingView::FindLastEvent(bool bUseCustomFilter)
 
 void STimingView::FindMaxEvent()
 {
+	INSIGHTS_SCOPED_LOG(FindMaxEvent);
+
 	if (SelectedEvent.IsValid())
 	{
 		SelectedEvent.Reset();
 		OnSelectedEventChangedDelegate.Broadcast(SelectedEvent);
 	}
 
+	const double MaxValidTime = Viewport.GetMaxValidTime();
 	TSharedPtr<const ITimingEvent> BestMatchEvent;
 	double CurrentDuration = -1.0;
 
 	TSharedPtr<FFilterConfigurator> Filter = QuickFindVm->GetFilterConfigurator();
 
-	auto EventFilter = [&CurrentDuration](double EventStartTime, double EventEndTime, uint32 EventDepth)
+	auto EventFilter = [&CurrentDuration, MaxValidTime](double EventStartTime, double EventEndTime, uint32 EventDepth)
 	{
-		const double Duration = EventEndTime - EventStartTime;
+		const double Duration = EventEndTime == std::numeric_limits<double>::infinity() ?
+			MaxValidTime - EventStartTime : EventEndTime - EventStartTime;
 		return Duration > CurrentDuration;
 	};
 
-	auto EventMatched = [&CurrentDuration](double EventStartTime, double EventEndTime, uint32 EventDepth)
+	auto EventMatched = [&CurrentDuration, MaxValidTime](double EventStartTime, double EventEndTime, uint32 EventDepth)
 	{
-		const double Duration = EventEndTime - EventStartTime;
+		const double Duration = EventEndTime == std::numeric_limits<double>::infinity() ?
+			MaxValidTime - EventStartTime : EventEndTime - EventStartTime;
 		CurrentDuration = Duration;
 	};
 
@@ -5730,26 +5775,31 @@ void STimingView::FindMaxEvent()
 
 void STimingView::FindMinEvent()
 {
+	INSIGHTS_SCOPED_LOG(FindMinEvent);
+
 	if (SelectedEvent.IsValid())
 	{
 		SelectedEvent.Reset();
 		OnSelectedEventChangedDelegate.Broadcast(SelectedEvent);
 	}
 
+	const double MaxValidTime = Viewport.GetMaxValidTime();
 	TSharedPtr<const ITimingEvent> BestMatchEvent;
 	double CurrentDuration = std::numeric_limits<double>::max();
 
 	TSharedPtr<FFilterConfigurator> Filter = QuickFindVm->GetFilterConfigurator();
 
-	auto EventFilter = [&CurrentDuration](double EventStartTime, double EventEndTime, uint32 EventDepth)
+	auto EventFilter = [&CurrentDuration, MaxValidTime](double EventStartTime, double EventEndTime, uint32 EventDepth)
 	{
-		const double Duration = EventEndTime - EventStartTime;
+		const double Duration = EventEndTime == std::numeric_limits<double>::infinity() ?
+			MaxValidTime - EventStartTime : EventEndTime - EventStartTime;
 		return Duration < CurrentDuration;
 	};
 
-	auto EventMatched = [&CurrentDuration](double EventStartTime, double EventEndTime, uint32 EventDepth)
+	auto EventMatched = [&CurrentDuration, MaxValidTime](double EventStartTime, double EventEndTime, uint32 EventDepth)
 	{
-		const double Duration = EventEndTime - EventStartTime;
+		const double Duration = EventEndTime == std::numeric_limits<double>::infinity() ?
+			MaxValidTime - EventStartTime : EventEndTime - EventStartTime;
 		CurrentDuration = Duration;
 	};
 
@@ -6015,7 +6065,9 @@ ETraceFrameType STimingView::GetFrameTypeToSnapTo()
 
 void STimingView::SelectEventInstance(uint32 TimerId, ESelectEventType Type, bool bUseSelection)
 {
-	double IntervalStart = 0.0f;
+	INSIGHTS_SCOPED_LOG(FindEvent);
+
+	double IntervalStart = 0.0;
 	double IntervalEnd = std::numeric_limits<double>::max();
 
 	if (bUseSelection && SelectionEndTime > SelectionStartTime)
@@ -6024,7 +6076,10 @@ void STimingView::SelectEventInstance(uint32 TimerId, ESelectEventType Type, boo
 		IntervalEnd = SelectionEndTime;
 	}
 
+	const double MaxValidTime = Viewport.GetMaxValidTime();
+
 	TSharedPtr<const ITimingEvent> TimingEvent;
+
 	// Create a temporary filter object using the TimerId and search range as filter criteria
 	MakeCustomFilterConfigurator(TimerId, IntervalStart, IntervalEnd);
 
@@ -6048,11 +6103,11 @@ void STimingView::SelectEventInstance(uint32 TimerId, ESelectEventType Type, boo
 		break;
 	case ESelectEventType::Min:
 		SelectTimingEvent(nullptr, false, false);
-		TimingEvent = ThreadTimingSharedState->FindMinEventInstance(TimerId, IntervalStart, IntervalEnd);
+		TimingEvent = ThreadTimingSharedState->FindMinEventInstance(TimerId, IntervalStart, IntervalEnd, MaxValidTime);
 		break;
 	case ESelectEventType::Max:
 		SelectTimingEvent(nullptr, false, false);
-		TimingEvent = ThreadTimingSharedState->FindMaxEventInstance(TimerId, IntervalStart, IntervalEnd);
+		TimingEvent = ThreadTimingSharedState->FindMaxEventInstance(TimerId, IntervalStart, IntervalEnd, MaxValidTime);
 		break;
 	}
 
@@ -6122,4 +6177,5 @@ TSharedPtr<SLogView> STimingView::GetLogView() const
 } // namespace UE::Insights::TimingProfiler
 
 #undef INSIGHTS_ACTIVATE_BENCHMARK
+#undef INSIGHTS_SCOPED_LOG
 #undef LOCTEXT_NAMESPACE
