[Insights] Add support for exporting frame tracks

diff --git a/Engine/Source/Developer/TraceInsights/Private/Insights/TimingProfiler/ViewModels/FrameTimingTrack.cpp b/Engine/Source/Developer/TraceInsights/Private/Insights/TimingProfiler/ViewModels/FrameTimingTrack.cpp
--- a/Engine/Source/Developer/TraceInsights/Private/Insights/TimingProfiler/ViewModels/FrameTimingTrack.cpp
+++ b/Engine/Source/Developer/TraceInsights/Private/Insights/TimingProfiler/ViewModels/FrameTimingTrack.cpp
@@ -128,7 +128,7 @@ void FFrameSharedState::Tick(Timing::ITimingViewSession& InSession, const TraceS
 			TSharedPtr<FFrameTimingTrack>* TrackPtrPtr = FrameTracks.Find(FrameType);
 			if (!TrackPtrPtr)
 			{
-				const FString TrackName = FString::Printf(TEXT("%s Frames"), FFrameTrackDrawHelper::FrameTypeToString(FrameType));
+				const FString TrackName = FFrameTimingTrack::GetFrameTrackName(FrameType);
 				TSharedPtr<FFrameTimingTrack> Track = MakeShared<FFrameTimingTrack>(*this, TrackName, FrameType);
 				Track->SetOrder(Order);
 				FrameTracks.Add(FrameType, Track);
@@ -413,7 +413,7 @@ void FFrameTimingTrack::InitTooltip(FTooltipDrawState& InOutTooltip, const ITimi
 		{
 			InOutTooltip.ResetContent();
 
-			const FString FrameName = GetFrameName(InFoundFrame.Index);
+			const FString FrameName = GetFrameName(FrameType, InFoundFrame.Index);
 			InOutTooltip.AddTitle(FrameName);
 
 			//const double Duration = TooltipEvent.GetDuration();
@@ -460,7 +460,7 @@ void FFrameTimingTrack::OnClipboardCopyEvent(const ITimingEvent& InSelectedEvent
 		const FTimingEvent& TrackEvent = InSelectedEvent.As<FTimingEvent>();
 
 		const uint64 FrameIndex = TrackEvent.GetType(); // using event type as frame index
-		const FString FrameName = GetFrameName(FrameIndex);
+		const FString FrameName = GetFrameName(FrameType, FrameIndex);
 
 		// Copy name of selected frame to clipboard.
 		FPlatformApplicationMisc::ClipboardCopy(*FrameName);
@@ -563,7 +563,7 @@ const FString FFrameTimingTrack::GetShortFrameName(const uint64 InFrameIndex) co
 
 ////////////////////////////////////////////////////////////////////////////////////////////////////
 
-const FString FFrameTimingTrack::GetFrameName(const uint64 InFrameIndex) const
+FString FFrameTimingTrack::GetFrameName(const uint32 FrameType, const uint64 InFrameIndex)
 {
 	return FString::Printf(TEXT("%s Frame %" UINT64_FMT), FFrameTrackDrawHelper::FrameTypeToString(FrameType), InFrameIndex);
 }
@@ -577,6 +577,13 @@ const FString FFrameTimingTrack::GetCompleteFrameName(const uint64 InFrameIndex,
 
 ////////////////////////////////////////////////////////////////////////////////////////////////////
 
+FString FFrameTimingTrack::GetFrameTrackName(uint32 FrameType)
+{
+	return FString::Printf(TEXT("%s Frames"), FFrameTrackDrawHelper::FrameTypeToString(FrameType));
+}
+
+////////////////////////////////////////////////////////////////////////////////////////////////////
+
 } // namespace UE::Insights::TimingProfiler
 
 #undef LOCTEXT_NAMESPACE
diff --git a/Engine/Source/Developer/TraceInsights/Private/Insights/TimingProfiler/ViewModels/FrameTimingTrack.h b/Engine/Source/Developer/TraceInsights/Private/Insights/TimingProfiler/ViewModels/FrameTimingTrack.h
--- a/Engine/Source/Developer/TraceInsights/Private/Insights/TimingProfiler/ViewModels/FrameTimingTrack.h
+++ b/Engine/Source/Developer/TraceInsights/Private/Insights/TimingProfiler/ViewModels/FrameTimingTrack.h
@@ -107,9 +107,11 @@ public:
 	void ToggleCollapsed() { Header.ToggleCollapsed(); }
 
 	const FString GetShortFrameName(const uint64 FrameIndex) const;
-	const FString GetFrameName(const uint64 FrameIndex) const;
+	static FString GetFrameName(const uint32 FrameType, const uint64 FrameIndex);
 	const FString GetCompleteFrameName(const uint64 FrameIndex, const double FrameDuration) const;
 
+	static FString GetFrameTrackName(uint32 FrameType);
+
 	//////////////////////////////////////////////////
 	// FBaseTimingTrack/FTimingEventsTrack overrides
 
diff --git a/Engine/Source/Developer/TraceInsights/Private/Insights/TimingProfiler/ViewModels/TimingExporter.cpp b/Engine/Source/Developer/TraceInsights/Private/Insights/TimingProfiler/ViewModels/TimingExporter.cpp
--- a/Engine/Source/Developer/TraceInsights/Private/Insights/TimingProfiler/ViewModels/TimingExporter.cpp
+++ b/Engine/Source/Developer/TraceInsights/Private/Insights/TimingProfiler/ViewModels/TimingExporter.cpp
@@ -22,6 +22,8 @@
 #include "Insights/TimingProfiler/TimingProfilerManager.h"
 #include "Insights/TimingProfiler/Tracks/GpuTimingTrack.h"
 #include "Insights/TimingProfiler/Tracks/VerseTimingTrack.h"
+#include "Insights/TimingProfiler/ViewModels/FrameTimingTrack.h"
+#include "Insights/TimingProfiler/ViewModels/FrameTrackHelper.h"
 
 #define LOCTEXT_NAMESPACE "UE::Insights::TimingProfiler::FTimingExporter"
 
@@ -174,6 +176,22 @@ int32 FTimingExporter::ExportThreadsAsText(const FString& Filename, FExportThrea
 		Writer.AppendLineEnd();
 		++ThreadCount;
 
+		// Iterate the Frame tracks.
+		{
+			TraceServices::FAnalysisSessionReadScope SessionReadScope(Session);
+			const TraceServices::IFrameProvider& FramesProvider = TraceServices::ReadFrameProvider(Session);
+			for (uint32 FrameType = 0; FrameType < (uint32)ETraceFrameType::TraceFrameType_Count; ++FrameType)
+			{
+				StringBuilder.Appendf(UTF8TEXT("%u"), GetFrameTypeNonCollidingId(FrameType));
+				Writer.AppendSeparator();
+				StringBuilder.Append(FFrameTimingTrack::GetFrameTrackName(FrameType));
+				Writer.AppendSeparator();
+				StringBuilder.Append(UTF8TEXT("Frame Tracks"));
+				Writer.AppendLineEnd();
+				++ThreadCount;
+			}
+		}
+
 		// Iterate the Gpu Queues.
 		{
 			TraceServices::FAnalysisSessionReadScope SessionReadScope(Session);
@@ -532,6 +550,62 @@ int32 FTimingExporter::ExportTimingEvents_WriteEvents(FExportTimingEventsInterna
 				});
 		};
 
+		auto FrameEnumerator = [&Params, &TimingEventCount](const TraceServices::FFrame& Frame)
+		{
+			FUtf8StringBuilder& StringBuilder = Params.Writer.GetStringBuilder();
+			bool bFirst = true;
+			for (const FName& Column : Params.Columns)
+			{
+				if (Params.Exporter.ExportTimingEventsColumns.Contains(Column))
+				{
+					if (bFirst)
+					{
+						bFirst = false;
+					}
+					else
+					{
+						Params.Writer.AppendSeparator();
+					}
+
+					if (Column == ExportTimingEvents_ThreadIdColumn)
+					{
+						StringBuilder.Appendf(UTF8TEXT("%u"), Params.ThreadId);
+					}
+					else if (Column == ExportTimingEvents_ThreadNameColumn)
+					{
+						Params.Writer.AppendString(Params.ThreadName);
+					}
+					else if (Column == ExportTimingEvents_TimerIdColumn)
+					{
+						uint32 TimerIndex = (uint32) Frame.Index;
+						StringBuilder.Appendf(UTF8TEXT("%u"), TimerIndex);
+					}
+					else if (Column == ExportTimingEvents_TimerNameColumn)
+					{
+						Params.Writer.AppendString(*FFrameTimingTrack::GetFrameName((uint32)Frame.FrameType, Frame.Index));
+					}
+					else if (Column == ExportTimingEvents_StartTimeColumn)
+					{
+						StringBuilder.Appendf(UTF8TEXT("%.9g"), Frame.StartTime);
+					}
+					else if (Column == ExportTimingEvents_EndTimeColumn)
+					{
+						StringBuilder.Appendf(UTF8TEXT("%.9g"), Frame.EndTime);
+					}
+					else if (Column == ExportTimingEvents_DurationColumn)
+					{
+						StringBuilder.Appendf(UTF8TEXT("%.9f"), Frame.EndTime - Frame.StartTime);
+					}
+					else if (Column == ExportTimingEvents_DepthColumn)
+					{
+						StringBuilder.Appendf(UTF8TEXT("%u"), 1u);
+					}
+				}
+			}
+			Params.Writer.AppendLineEnd();
+			++TimingEventCount;
+		};
+
 		// Iterate the GPU timelines.
 		{
 			if (!Params.UserParams.ThreadFilter || Params.UserParams.ThreadFilter(FGpuTimingTrack::Gpu1ThreadId))
@@ -579,6 +653,22 @@ int32 FTimingExporter::ExportTimingEvents_WriteEvents(FExportTimingEventsInterna
 			}
 		}
 
+		// Iterate the Frame timelines.
+		{
+			const TraceServices::IFrameProvider& FramesProvider = TraceServices::ReadFrameProvider(Session);
+			for (uint32 FrameType = 0; FrameType < (uint32)ETraceFrameType::TraceFrameType_Count; ++FrameType)
+			{
+				if (!Params.UserParams.ThreadFilter || Params.UserParams.ThreadFilter(GetFrameTypeNonCollidingId(FrameType)))
+				{
+					FString ThreadName = FFrameTimingTrack::GetFrameTrackName(FrameType);
+					Params.ThreadId = GetFrameTypeNonCollidingId(FrameType);
+					Params.ThreadName = *ThreadName;
+					FramesProvider.EnumerateFrames((ETraceFrameType) FrameType, Params.UserParams.IntervalStartTime, Params.UserParams.IntervalEndTime, FrameEnumerator);
+					Params.ThreadName = nullptr;
+				}
+			}
+		}
+
 		// Iterate the CPU threads and their corresponding timelines.
 		const TraceServices::IThreadProvider& ThreadProvider = TraceServices::ReadThreadProvider(Session);
 		ThreadProvider.EnumerateThreads(
@@ -1170,6 +1260,16 @@ FTimingExporter::FThreadFilterFunc FTimingExporter::MakeThreadFilterInclusive(co
 			});
 	}
 
+	// Add the frame tracks
+	{
+		TraceServices::FAnalysisSessionReadScope SessionReadScope(Session);
+		const TraceServices::IFrameProvider& FramesProvider = TraceServices::ReadFrameProvider(Session);
+		for (uint32 FrameType = 0; FrameType < (uint32)ETraceFrameType::TraceFrameType_Count; ++FrameType)
+		{
+			Threads.Add(FFrameTimingTrack::GetFrameTrackName(FrameType), GetFrameTypeNonCollidingId(FrameType));
+		}
+	}
+
 	TArray<FString> Filter;
 	InFilterString.ParseIntoArray(Filter, TEXT(","), true);
 
@@ -1687,6 +1787,13 @@ uint32 FTimingExporter::GetNonCollidingId(uint32 QueueId)
 
 ////////////////////////////////////////////////////////////////////////////////////////////////////
 
+uint32 FTimingExporter::GetFrameTypeNonCollidingId(uint32 FrameType)
+{
+	return FrameType + (1 << 24);
+}
+
+////////////////////////////////////////////////////////////////////////////////////////////////////
+
 } // namespace UE::Insights::TimingProfiler
 
 #undef LOCTEXT_NAMESPACE
diff --git a/Engine/Source/Developer/TraceInsights/Private/Insights/TimingProfiler/ViewModels/TimingExporter.h b/Engine/Source/Developer/TraceInsights/Private/Insights/TimingProfiler/ViewModels/TimingExporter.h
--- a/Engine/Source/Developer/TraceInsights/Private/Insights/TimingProfiler/ViewModels/TimingExporter.h
+++ b/Engine/Source/Developer/TraceInsights/Private/Insights/TimingProfiler/ViewModels/TimingExporter.h
@@ -335,6 +335,7 @@ private:
 	int32 ExportTimerCalleesByRegions(const FString& FilenamePattern, const FExportTimerCalleesParams& Params) const;
 	int32 ExportCounterAsTextByRegions(const FString& Filename, uint32 CounterId, FExportCounterParams& Params) const;
 	static uint32 GetNonCollidingId(uint32 QueueId);
+	static uint32 GetFrameTypeNonCollidingId(uint32 FrameType);
 
 private:
 	const TraceServices::IAnalysisSession& Session;
