[Insights] Fixed "Find Timing Event Instance" to properly bring the track/lane of the selected timing event into view. Also enabled auto zooming on the found timing event.

diff --git a/Engine/Source/Developer/TraceInsights/Private/Insights/Widgets/STimingView.cpp b/Engine/Source/Developer/TraceInsights/Private/Insights/Widgets/STimingView.cpp
--- a/Engine/Source/Developer/TraceInsights/Private/Insights/Widgets/STimingView.cpp
+++ b/Engine/Source/Developer/TraceInsights/Private/Insights/Widgets/STimingView.cpp
@@ -1045,39 +1045,85 @@ PRAGMA_ENABLE_DEPRECATION_WARNINGS
 
 	////////////////////////////////////////////////////////////////////////////////////////////////////
 
-	const float InitialScrollPosY = Viewport.GetScrollPosY();
-
-	TSharedPtr<FBaseTimingTrack> SelectedScrollableTrack;
-	if (SelectedTrack.IsValid() && SelectedTrack->IsVisible())
+	TSharedPtr<FBaseTimingTrack> PinnedTrack;
+	float InitialPinnedTrackPosY = 0.0f;
+	if (SelectedTrack.IsValid() && SelectedTrack->IsVisible() && DelayedVisibleTrackCounter <= 0)
 	{
 		if (ScrollableTracks.Contains(SelectedTrack))
 		{
-			SelectedScrollableTrack = SelectedTrack;
+			PinnedTrack = SelectedTrack;
+			InitialPinnedTrackPosY = PinnedTrack->GetPosY();
 		}
 	}
 
-	const float InitialPinnedTrackPosY = SelectedScrollableTrack.IsValid() ? SelectedScrollableTrack->GetPosY() : 0.0f;
-
 	// Update the Y position for visible scrollable tracks.
 	UpdatePositionForScrollableTracks();
 
 	// The selected track will be pinned (keeps Y pos fixed unless user scrolls vertically).
-	if (SelectedScrollableTrack.IsValid())
+	if (PinnedTrack.IsValid())
 	{
-		const float ScrollingDY = LastScrollPosY - InitialScrollPosY;
-		const float PinnedTrackPosY = SelectedScrollableTrack->GetPosY();
-		const float AdjustmentDY = InitialPinnedTrackPosY - PinnedTrackPosY + ScrollingDY;
-
+		const float InitialScrollPosY = Viewport.GetScrollPosY();
+		const float AdjustmentDY = LastScrollPosY - InitialScrollPosY + InitialPinnedTrackPosY - PinnedTrack->GetPosY();
+		const float NewScrollPosY = InitialScrollPosY - AdjustmentDY;
 		if (!FMath::IsNearlyZero(AdjustmentDY, 0.5f))
 		{
 			ViewportScrollPosYOnButtonDown -= AdjustmentDY;
-			ScrollAtPosY(InitialScrollPosY - AdjustmentDY);
+			ScrollAtPosY(NewScrollPosY);
 			UpdatePositionForScrollableTracks();
 		}
 	}
 
+	// Force the DelayedVisibleTrack to be visible.
+	if (DelayedVisibleTrackCounter > 0)
+	{
+		if (DelayedVisibleTrack.IsValid() &&
+			DelayedVisibleTrack->IsVisible() &&
+			DelayedVisibleTrack->GetLocation() == ETimingTrackLocation::Scrollable)
+		{
+			float PreviousScrollPosY = Viewport.GetScrollPosY();
+
+			// Try to bring into view the entire track (with 10% top/bottom margins).
+			const float TrackHeight = DelayedVisibleTrack->GetHeight();
+			const float TopScrollY = DelayedVisibleTrack->GetPosY() + Viewport.GetScrollPosY() - Viewport.GetTopOffset() - Viewport.GetPosY();
+			const float ScrollableAreaHeight = Viewport.GetScrollableAreaHeight();
+			BringIntoViewY(TopScrollY - ScrollableAreaHeight * 0.1f, TopScrollY + TrackHeight + ScrollableAreaHeight * 0.2f);
+
+			// Try to bring into view the entire track.
+			BringIntoViewY(TopScrollY, TopScrollY + TrackHeight);
+
+			if (DelayedVisibleTrackLaneDepth >= 0)
+			{
+				// Try to bring into view only the lane at specified depth (with 10% top/bottom margins).
+				const float LaneHeight = Viewport.GetLayout().EventH + Viewport.GetLayout().EventDY;
+				const float ScrollY = TopScrollY + float(DelayedVisibleTrackLaneDepth) * LaneHeight;
+				BringIntoViewY(ScrollY - ScrollableAreaHeight * 0.1f, ScrollY + LaneHeight + ScrollableAreaHeight * 0.2f);
+			}
+
+			UpdatePositionForScrollableTracks();
+
+			if (FMath::IsNearlyZero(Viewport.GetScrollPosY() - PreviousScrollPosY, 0.5f) &&
+				FMath::IsNearlyZero(Viewport.GetScrollHeight() - DelayedVisibleTrackLastScrollHeight, 0.5f))
+			{
+				if (--DelayedVisibleTrackCounter == 0)
+				{
+					DelayedVisibleTrack.Reset();
+				}
+			}
+			else
+			{
+				DelayedVisibleTrackCounter = 3;
+				DelayedVisibleTrackLastScrollHeight = Viewport.GetScrollHeight();
+			}
+		}
+		else
+		{
+			DelayedVisibleTrack.Reset();
+			DelayedVisibleTrackCounter = 0;
+		}
+	}
+
 	// Elastic snap to vertical scroll limits.
-	if (!bIsPanning)
+	if (!bIsPanning && DelayedVisibleTrackCounter <= 0)
 	{
 		const float DY = Viewport.GetScrollHeight() - Viewport.GetScrollableAreaHeight();
 		const float MinY = FMath::Min(DY, 0.0f);
@@ -3623,6 +3669,16 @@ void STimingView::BringScrollableTrackIntoView(const FBaseTimingTrack& InTrack)
 
 ////////////////////////////////////////////////////////////////////////////////////////////////////
 
+void STimingView::BringScrollableTrackIntoViewDelayed(TSharedPtr<FBaseTimingTrack> Track, int32 LaneDepth)
+{
+	DelayedVisibleTrack = Track;
+	DelayedVisibleTrackLaneDepth = LaneDepth;
+	DelayedVisibleTrackCounter = 3;
+	DelayedVisibleTrackLastScrollHeight = 0.0f;
+}
+
+////////////////////////////////////////////////////////////////////////////////////////////////////
+
 void STimingView::ScrollAtTime(double StartTime)
 {
 	if (Viewport.ScrollAtTime(StartTime))
@@ -5836,12 +5892,19 @@ void STimingView::PostFindEvent(TSharedPtr<const ITimingEvent> BestMatchEvent)
 	if (BestMatchEvent)
 	{
 		SelectedEvent = BestMatchEvent;
-		BringIntoView(SelectedEvent->GetStartTime(), SelectedEvent->GetEndTime());
-		if (SelectedEvent->GetTrack()->GetLocation() == ETimingTrackLocation::Scrollable)
+		bool bZoomOnFoundEvent = true;
+		if (bZoomOnFoundEvent)
 		{
-			BringScrollableTrackIntoView(*SelectedEvent->GetTrack());
+			const double Duration = Viewport.RestrictDuration(SelectedEvent->GetStartTime(), SelectedEvent->GetEndTime());
+			if (Duration > 0.0)
+			{
+				ZoomOnTimeInterval(SelectedEvent->GetStartTime() - Duration * 0.1, Duration * 1.2);
+			}
 		}
-
+		BringIntoView(SelectedEvent->GetStartTime(), SelectedEvent->GetEndTime());
+		TSharedPtr<FBaseTimingTrack> Track = ConstCastSharedPtr<FBaseTimingTrack>(SelectedEvent->GetTrack().ToSharedPtr());
+		BringScrollableTrackIntoViewDelayed(Track, int32(SelectedEvent->GetDepth()));
+		SelectTimingTrack(Track, true);
 		OnSelectedTimingEventChanged();
 	}
 	else
diff --git a/Engine/Source/Developer/TraceInsights/Private/Insights/Widgets/STimingView.h b/Engine/Source/Developer/TraceInsights/Private/Insights/Widgets/STimingView.h
--- a/Engine/Source/Developer/TraceInsights/Private/Insights/Widgets/STimingView.h
+++ b/Engine/Source/Developer/TraceInsights/Private/Insights/Widgets/STimingView.h
@@ -392,6 +392,7 @@ public:
 	void ScrollAtPosY(float ScrollPosY);
 	void BringIntoViewY(float InTopY, float InBottomY);
 	void BringScrollableTrackIntoView(const FBaseTimingTrack& Track);
+	void BringScrollableTrackIntoViewDelayed(TSharedPtr<FBaseTimingTrack> Track, int32 LaneDepth = -1);
 	void ScrollAtTime(double StartTime);
 	void CenterOnTimeInterval(double IntervalStartTime, double IntervalDuration);
 	void ZoomOnTimeInterval(double IntervalStartTime, double IntervalDuration);
@@ -742,6 +743,11 @@ protected:
 	TSharedPtr<FBaseTimingTrack> SelectedTrack;
 	TSharedPtr<const ITimingEvent> SelectedEvent;
 
+	TSharedPtr<FBaseTimingTrack> DelayedVisibleTrack;
+	int32 DelayedVisibleTrackLaneDepth = -1;
+	float DelayedVisibleTrackLastScrollHeight = 0.0f;
+	int32 DelayedVisibleTrackCounter = 0;
+
 	TSharedPtr<ITimingEventFilter> TimingEventFilter;
 
 	FTooltipDrawState Tooltip;
