[Insights]
* InsightsCore: Added FBaseFilterContext class to filtering API to allow custom implementations of the filter context.
* Improved performance for event searching and filtering in the timing event tracks (now using custom FilterContext classes). For example, using a large trace, the Find First/Last Timing Event Instance operations from the context menu of Timers panel improved with -43% (from 37s to 21s) and equivalent Find First/Last operations from Quick Find panel improved with -58% (from 29s to 12s).
* InsightsCore: Fixed logical OR group filter operator over an empty set to evaluate to false to avoid unintended matches.

#Synced-CL 48209574

diff --git a/Engine/Source/Developer/TraceInsights/Private/Insights/ContextSwitches/ViewModels/CpuCoreTimingTrack.cpp b/Engine/Source/Developer/TraceInsights/Private/Insights/ContextSwitches/ViewModels/CpuCoreTimingTrack.cpp
--- a/Engine/Source/Developer/TraceInsights/Private/Insights/ContextSwitches/ViewModels/CpuCoreTimingTrack.cpp
+++ b/Engine/Source/Developer/TraceInsights/Private/Insights/ContextSwitches/ViewModels/CpuCoreTimingTrack.cpp
@@ -26,6 +26,102 @@
 namespace UE::Insights::ContextSwitches
 {
 
+////////////////////////////////////////////////////////////////////////////////////////////////////
+// FCpuCoreTimingTrackFilterContext
+////////////////////////////////////////////////////////////////////////////////////////////////////
+
+class FCpuCoreTimingTrackFilterContext : public FBaseFilterContext
+{
+public:
+	FCpuCoreTimingTrackFilterContext(const FCpuCoreTimingTrack& InTrack)
+		: Track(InTrack)
+	{
+	}
+	virtual ~FCpuCoreTimingTrackFilterContext() = default;
+
+	virtual bool HasData(int32 Key) const override
+	{
+		switch (Key)
+		{
+		case static_cast<int32>(EFilterField::StartTime):
+		case static_cast<int32>(EFilterField::EndTime):
+		case static_cast<int32>(EFilterField::Duration):
+		case static_cast<int32>(EFilterField::TrackName):
+		case static_cast<int32>(EFilterField::CoreEventName):
+			return true;
+		default:
+			return false;
+		}
+	}
+
+	virtual EFilterDataType GetDataType(int32 Key) const override
+	{
+		switch (Key)
+		{
+		case static_cast<int32>(EFilterField::StartTime):
+		case static_cast<int32>(EFilterField::EndTime):
+		case static_cast<int32>(EFilterField::Duration):
+			return EFilterDataType::Double;
+		case static_cast<int32>(EFilterField::TrackName):
+			return EFilterDataType::String;
+		case static_cast<int32>(EFilterField::CoreEventName): // exposed as SystemThreadId
+			return EFilterDataType::Int64;
+		default:
+			return EFilterDataType::Void;
+		}
+	}
+
+	virtual int64 GetDataAsInt64(int32 Key) const override
+	{
+		if (Key == static_cast<int32>(EFilterField::CoreEventName))
+		{
+			return int64(SystemThreadId);
+		}
+		return 0;
+	}
+
+	virtual double GetDataAsDouble(int32 Key) const override
+	{
+		if (Key == static_cast<int32>(EFilterField::StartTime))
+		{
+			return EventStartTime;
+		}
+		if (Key == static_cast<int32>(EFilterField::EndTime))
+		{
+			return EventEndTime;
+		}
+		if (Key == static_cast<int32>(EFilterField::Duration))
+		{
+			return EventEndTime - EventStartTime;
+		}
+		return 0.0;
+	}
+
+	virtual const FString& GetDataAsString(int32 Key) const override
+	{
+		if (Key == static_cast<int32>(EFilterField::TrackName))
+		{
+			return Track.GetName();
+		}
+		return EmptyString;
+	}
+
+	void SetEvent(const TraceServices::FCpuCoreEvent& Event)
+	{
+		EventStartTime = Event.Start;
+		EventEndTime = Event.End;
+		SystemThreadId = Event.SystemThreadId;
+	}
+
+private:
+	const FCpuCoreTimingTrack& Track;
+	double EventStartTime = 0.0;
+	double EventEndTime = 0.0;
+	uint32 SystemThreadId = 0;
+};
+
+////////////////////////////////////////////////////////////////////////////////////////////////////
+// FCpuCoreTimingTrack
 ////////////////////////////////////////////////////////////////////////////////////////////////////
 
 INSIGHTS_IMPLEMENT_RTTI(FCpuCoreTimingTrack)
@@ -93,6 +189,12 @@ void FCpuCoreTimingTrack::BuildDrawState(ITimingEventsTrackDrawStateBuilder& Bui
 
 void FCpuCoreTimingTrack::BuildFilteredDrawState(ITimingEventsTrackDrawStateBuilder& Builder, const ITimingTrackUpdateContext& Context)
 {
+	TSharedPtr<const TraceServices::IAnalysisSession> Session = FInsightsManager::Get()->GetSession();
+	if (!Session.IsValid())
+	{
+		return;
+	}
+
 	const TSharedPtr<ITimingEventFilter> EventFilterPtr = Context.GetEventFilter();
 	if (EventFilterPtr.IsValid() && EventFilterPtr->FilterTrack(*this))
 	{
@@ -105,90 +207,19 @@ void FCpuCoreTimingTrack::BuildFilteredDrawState(ITimingEventsTrackDrawStateBuil
 			FilterEventType = EventFilter.GetEventType();
 		}
 
-		TSharedPtr<const TraceServices::IAnalysisSession> Session = FInsightsManager::Get()->GetSession();
-		if (Session.IsValid())
-		{
-			if (const TraceServices::IContextSwitchesProvider* ContextSwitchesProvider = TraceServices::ReadContextSwitchesProvider(*Session.Get()))
-			{
-				TraceServices::FProviderReadScopeLock ScopedLock(*ContextSwitchesProvider);
-
-				const FTimingTrackViewport& Viewport = Context.GetViewport();
-
-				if (bFilterOnlyByEventType)
-				{
-					bool bShowNonTargetProcessEvents = SharedState.AreNonTargetProcessEventsVisible();
-					ContextSwitchesProvider->EnumerateCpuCoreEvents(CoreNumber, Viewport.GetStartTime(), Viewport.GetEndTime(),
-						[this, &Builder, ContextSwitchesProvider, FilterEventType, bShowNonTargetProcessEvents]
-						(const TraceServices::FCpuCoreEvent& CpuCoreEvent)
-						{
-							if (CpuCoreEvent.SystemThreadId != FilterEventType)
-							{
-								return TraceServices::EContextSwitchEnumerationResult::Continue;
-							}
-
-							uint32 ThreadId;
-							if (!bShowNonTargetProcessEvents && !ContextSwitchesProvider->GetThreadId(CpuCoreEvent.SystemThreadId, ThreadId))
-							{
-								return TraceServices::EContextSwitchEnumerationResult::Continue;
-							}
-
-							this->AddCoreTimingEvent(Builder, CpuCoreEvent);
-
-							return TraceServices::EContextSwitchEnumerationResult::Continue;
-						});
-				}
-				else // generic filter
-				{
-					//TODO: if (EventFilterPtr->FilterEvent(TimingEvent))
-				}
-			}
-		}
-	}
-
-	if (HasCustomFilter())
-	{
-		TSharedPtr<TimingProfiler::STimingView> TimingView = SharedState.GetTimingView();
-		if (!TimingView)
-		{
-			return;
-		}
-
-		TSharedPtr<FFilterConfigurator> FilterConfigurator = TimingView->GetFilterConfigurator();
-		if (!FilterConfigurator.IsValid())
-		{
-			return;
-		}
-
-		FFilterContext FilterContext;
-		FilterContext.SetReturnValueForUnsetFilters(false);
-
-		FilterContext.AddFilterData<double>(static_cast<int32>(EFilterField::StartTime), 0.0f);
-		FilterContext.AddFilterData<double>(static_cast<int32>(EFilterField::EndTime), 0.0f);
-		FilterContext.AddFilterData<double>(static_cast<int32>(EFilterField::Duration), 0.0f);
-		FilterContext.AddFilterData<FString>(static_cast<int32>(EFilterField::TrackName), this->GetName());
-		FilterContext.AddFilterData<int64>(static_cast<int32>(EFilterField::CoreEventName), 0);
-
-		TSharedPtr<const TraceServices::IAnalysisSession> Session = FInsightsManager::Get()->GetSession();
-		if (Session.IsValid())
+		if (bFilterOnlyByEventType)
 		{
 			if (const TraceServices::IContextSwitchesProvider* ContextSwitchesProvider = TraceServices::ReadContextSwitchesProvider(*Session.Get()))
 			{
-				TraceServices::FProviderReadScopeLock ScopedLock(*ContextSwitchesProvider);
-
 				const FTimingTrackViewport& Viewport = Context.GetViewport();
-
 				bool bShowNonTargetProcessEvents = SharedState.AreNonTargetProcessEventsVisible();
 
+				TraceServices::FProviderReadScopeLock ScopedLock(*ContextSwitchesProvider);
 				ContextSwitchesProvider->EnumerateCpuCoreEvents(CoreNumber, Viewport.GetStartTime(), Viewport.GetEndTime(),
-					[this, &Builder, ContextSwitchesProvider, &FilterContext, &FilterConfigurator, bShowNonTargetProcessEvents]
+					[this, &Builder, ContextSwitchesProvider, FilterEventType, bShowNonTargetProcessEvents]
 					(const TraceServices::FCpuCoreEvent& CpuCoreEvent)
 					{
-						FilterContext.SetFilterData<double>(static_cast<int32>(EFilterField::StartTime), CpuCoreEvent.Start);
-						FilterContext.SetFilterData<double>(static_cast<int32>(EFilterField::EndTime), CpuCoreEvent.End);
-						FilterContext.SetFilterData<double>(static_cast<int32>(EFilterField::Duration), CpuCoreEvent.End - CpuCoreEvent.Start);
-						FilterContext.SetFilterData<int64>(static_cast<int32>(EFilterField::CoreEventName), CpuCoreEvent.SystemThreadId);
-
-						if (!FilterConfigurator->ApplyFilters(FilterContext))
+						if (CpuCoreEvent.SystemThreadId != FilterEventType)
 						{
 							return TraceServices::EContextSwitchEnumerationResult::Continue;
 						}
@@ -205,6 +236,49 @@ void FCpuCoreTimingTrack::BuildFilteredDrawState(ITimingEventsTrackDrawStateBuil
 					});
 			}
 		}
+		else // generic filter
+		{
+			//TODO: if (EventFilterPtr->FilterEvent(TimingEvent))
+		}
+	}
+
+	TSharedPtr<TimingProfiler::STimingView> TimingView = SharedState.GetTimingView();
+	if (!TimingView)
+	{
+		return;
+	}
+
+	TSharedPtr<FFilterConfigurator> FilterConfigurator = TimingView->GetFilterConfigurator();
+	if (FilterConfigurator.IsValid() && !FilterConfigurator->IsEmpty()) // HasCustomFilter()
+	{
+		if (const TraceServices::IContextSwitchesProvider* ContextSwitchesProvider = TraceServices::ReadContextSwitchesProvider(*Session.Get()))
+		{
+			FCpuCoreTimingTrackFilterContext FilterContext(*this);
+			const FTimingTrackViewport& Viewport = Context.GetViewport();
+			bool bShowNonTargetProcessEvents = SharedState.AreNonTargetProcessEventsVisible();
+
+			TraceServices::FProviderReadScopeLock ScopedLock(*ContextSwitchesProvider);
+			ContextSwitchesProvider->EnumerateCpuCoreEvents(CoreNumber, Viewport.GetStartTime(), Viewport.GetEndTime(),
+				[this, &Builder, ContextSwitchesProvider, &FilterContext, &FilterConfigurator, bShowNonTargetProcessEvents]
+				(const TraceServices::FCpuCoreEvent& CpuCoreEvent)
+				{
+					FilterContext.SetEvent(CpuCoreEvent);
+					if (!FilterConfigurator->ApplyFilters(FilterContext))
+					{
+						return TraceServices::EContextSwitchEnumerationResult::Continue;
+					}
+
+					uint32 ThreadId;
+					if (!bShowNonTargetProcessEvents && !ContextSwitchesProvider->GetThreadId(CpuCoreEvent.SystemThreadId, ThreadId))
+					{
+						return TraceServices::EContextSwitchEnumerationResult::Continue;
+					}
+
+					this->AddCoreTimingEvent(Builder, CpuCoreEvent);
+
+					return TraceServices::EContextSwitchEnumerationResult::Continue;
+				});
+		}
 	}
 }
 
@@ -349,20 +423,36 @@ const TSharedPtr<const ITimingEvent> FCpuCoreTimingTrack::SearchEvent(const FTim
 {
 	TSharedPtr<const ITimingEvent> FoundEvent;
 
-	FFilterContext FilterConfiguratorContext;
-	FilterConfiguratorContext.SetReturnValueForUnsetFilters(false);
-	FilterConfiguratorContext.AddFilterData<double>(static_cast<int32>(EFilterField::StartTime), 0.0f);
-	FilterConfiguratorContext.AddFilterData<double>(static_cast<int32>(EFilterField::EndTime), 0.0f);
-	FilterConfiguratorContext.AddFilterData<double>(static_cast<int32>(EFilterField::Duration), 0.0f);
-	FilterConfiguratorContext.AddFilterData<FString>(static_cast<int32>(EFilterField::TrackName), this->GetName());
-	FilterConfiguratorContext.AddFilterData<int64>(static_cast<int32>(EFilterField::CoreEventName), 0);
+	FindCpuCoreTimingEvent(InSearchParameters,
+		[this, &FoundEvent]
+		(double InFoundStartTime, double InFoundEndTime, uint32 InFoundDepth, const TraceServices::FCpuCoreEvent& InEvent)
+		{
+			FoundEvent = MakeShared<FTimingEvent>(SharedThis(this), InFoundStartTime, InFoundEndTime, InFoundDepth, InEvent.SystemThreadId);
+		});
+
+	return FoundEvent;
+}
+
+////////////////////////////////////////////////////////////////////////////////////////////////////
 
-	TTimingEventSearch<TraceServices::FCpuCoreEvent>::Search(
+bool FCpuCoreTimingTrack::FindCpuCoreTimingEvent(
+	const FTimingEventSearchParameters& InParameters,
+	TFunctionRef<void(double, double, uint32, const TraceServices::FCpuCoreEvent&)> InFoundPredicate) const
+{
+	TSharedPtr<const TraceServices::IAnalysisSession> Session = FInsightsManager::Get()->GetSession();
+	if (!Session.IsValid())
+	{
+		return false;
+	}
 
-		InSearchParameters,
+	FCpuCoreTimingTrackFilterContext FilterContext(*this);
+
+	return TTimingEventSearch<TraceServices::FCpuCoreEvent>::Search(
+
+		InParameters,
 
 		// Search Predicate
-		[this, &InSearchParameters]
+		[this, &InParameters]
 		(TTimingEventSearch<TraceServices::FCpuCoreEvent>::FContext& InContext)
 		{
 			TSharedPtr<const TraceServices::IAnalysisSession> Session = FInsightsManager::Get()->GetSession();
@@ -371,70 +461,49 @@ const TSharedPtr<const ITimingEvent> FCpuCoreTimingTrack::SearchEvent(const FTim
 				if (const TraceServices::IContextSwitchesProvider* ContextSwitchesProvider = TraceServices::ReadContextSwitchesProvider(*Session.Get()))
 				{
 					TraceServices::FProviderReadScopeLock ScopedLock(*ContextSwitchesProvider);
+
+					auto Callback =
+						[&InContext]
+						(const TraceServices::FCpuCoreEvent& CpuCoreEvent)
+						{
+							InContext.Check(CpuCoreEvent.Start, CpuCoreEvent.End, 0, CpuCoreEvent);
+							return InContext.ShouldContinueSearching() ?
+								TraceServices::EContextSwitchEnumerationResult::Continue :
+								TraceServices::EContextSwitchEnumerationResult::Stop;
+						};
+
 					if (InContext.GetParameters().SearchDirection == FTimingEventSearchParameters::ESearchDirection::Forward)
 					{
-						ContextSwitchesProvider->EnumerateCpuCoreEvents(CoreNumber, InSearchParameters.StartTime, InSearchParameters.EndTime,
-							[&InContext]
-							(const TraceServices::FCpuCoreEvent& CpuCoreEvent)
-							{
-								InContext.Check(CpuCoreEvent.Start, CpuCoreEvent.End, 0, CpuCoreEvent);
-								return InContext.ShouldContinueSearching() ?
-									TraceServices::EContextSwitchEnumerationResult::Continue :
-									TraceServices::EContextSwitchEnumerationResult::Stop;
-							});
+						ContextSwitchesProvider->EnumerateCpuCoreEvents(CoreNumber, InParameters.StartTime, InParameters.EndTime, Callback);
 					}
 					else
 					{
-						ContextSwitchesProvider->EnumerateCpuCoreEventsBackwards(CoreNumber, InSearchParameters.EndTime, InSearchParameters.StartTime,
-							[&InContext]
-							(const TraceServices::FCpuCoreEvent& CpuCoreEvent)
-							{
-								InContext.Check(CpuCoreEvent.Start, CpuCoreEvent.End, 0, CpuCoreEvent);
-								return InContext.ShouldContinueSearching() ?
-									TraceServices::EContextSwitchEnumerationResult::Continue :
-									TraceServices::EContextSwitchEnumerationResult::Stop;
-							});
+						ContextSwitchesProvider->EnumerateCpuCoreEventsBackwards(CoreNumber, InParameters.EndTime, InParameters.StartTime, Callback);
 					}
 				}
 			}
 		},
 
 		// Filter Predicate
-		[&FilterConfiguratorContext, &InSearchParameters]
+		[&FilterContext, &InParameters]
 		(double EventStartTime, double EventEndTime, uint32 EventDepth, const TraceServices::FCpuCoreEvent& Event)
 		{
-			if (!InSearchParameters.FilterExecutor.IsValid())
+			if (!InParameters.FilterExecutor.IsValid())
 			{
 				return true;
 			}
 
-			TSharedPtr<const TraceServices::IAnalysisSession> Session = FInsightsManager::Get()->GetSession();
-			if (!Session.IsValid())
-			{
-				return false;
-			}
-
-			FilterConfiguratorContext.SetFilterData<double>(static_cast<int32>(EFilterField::StartTime), EventStartTime);
-			FilterConfiguratorContext.SetFilterData<double>(static_cast<int32>(EFilterField::EndTime), EventEndTime);
-			FilterConfiguratorContext.SetFilterData<double>(static_cast<int32>(EFilterField::Duration), EventEndTime - EventStartTime);
-			FilterConfiguratorContext.SetFilterData<int64>(static_cast<int32>(EFilterField::CoreEventName), Event.SystemThreadId);
-
-			return InSearchParameters.FilterExecutor->ApplyFilters(FilterConfiguratorContext);
+			FilterContext.SetEvent(Event);
+			return InParameters.FilterExecutor->ApplyFilters(FilterContext);
 		},
 
 		// Found Predicate
-		[&FoundEvent, this]
-		(double InFoundStartTime, double InFoundEndTime, uint32 InFoundDepth, const TraceServices::FCpuCoreEvent& InEvent)
-		{
-			FoundEvent = MakeShared<FTimingEvent>(SharedThis(this), InFoundStartTime, InFoundEndTime, InFoundDepth, InEvent.SystemThreadId);
-		},
+		InFoundPredicate,
 
 		// Payload Matched Predicate
 		TTimingEventSearch<TraceServices::FCpuCoreEvent>::NoMatch
 
 	); // Search
-
-	return FoundEvent;
 }
 
 ////////////////////////////////////////////////////////////////////////////////////////////////////
diff --git a/Engine/Source/Developer/TraceInsights/Private/Insights/ContextSwitches/ViewModels/CpuCoreTimingTrack.h b/Engine/Source/Developer/TraceInsights/Private/Insights/ContextSwitches/ViewModels/CpuCoreTimingTrack.h
--- a/Engine/Source/Developer/TraceInsights/Private/Insights/ContextSwitches/ViewModels/CpuCoreTimingTrack.h
+++ b/Engine/Source/Developer/TraceInsights/Private/Insights/ContextSwitches/ViewModels/CpuCoreTimingTrack.h
@@ -36,31 +36,24 @@ public:
 	uint32 GetCoreNumber() const { return CoreNumber; }
 
 	virtual void Reset() override;
-
 	virtual void BuildDrawState(ITimingEventsTrackDrawStateBuilder& Builder, const ITimingTrackUpdateContext& Context) override;
 	virtual void BuildFilteredDrawState(ITimingEventsTrackDrawStateBuilder& Builder, const ITimingTrackUpdateContext& Context) override;
-
 	virtual void Draw(const ITimingTrackDrawContext& Context) const override;
-
 	virtual void PostDraw(const ITimingTrackDrawContext& Context) const override;
-
 	virtual void InitTooltip(FTooltipDrawState& InOutTooltip, const ITimingEvent& InTooltipEvent) const override;
-
 	virtual void BuildContextMenu(FMenuBuilder& MenuBuilder) override;
-
 	virtual const TSharedPtr<const ITimingEvent> SearchEvent(const FTimingEventSearchParameters& InSearchParameters) const override;
 
 protected:
 	virtual const TSharedPtr<const ITimingEvent> GetEvent(double InTime, double SecondsPerPixel, int32 Depth) const override;
-
 	virtual bool HasCustomFilter() const override;
+	virtual int32 GetMaxDepth() const override { return NonTargetProcessEventsMaxDepth; }
 
+private:
+	bool FindCpuCoreTimingEvent(const FTimingEventSearchParameters& InParameters, TFunctionRef<void(double, double, uint32, const TraceServices::FCpuCoreEvent&)> InFoundPredicate) const;
 	void AddCoreTimingEvent(ITimingEventsTrackDrawStateBuilder& Builder, const TraceServices::FCpuCoreEvent& CpuCoreEvent);
-
 	FString GetThreadName(uint32 InSystemThreadId) const;
 
-	virtual int32 GetMaxDepth() const override { return NonTargetProcessEventsMaxDepth; }
-
 private:
 	FContextSwitchesSharedState& SharedState;
 
diff --git a/Engine/Source/Developer/TraceInsights/Private/Insights/LoadingProfiler/Tracks/LoadingTimingTrack.cpp b/Engine/Source/Developer/TraceInsights/Private/Insights/LoadingProfiler/Tracks/LoadingTimingTrack.cpp
--- a/Engine/Source/Developer/TraceInsights/Private/Insights/LoadingProfiler/Tracks/LoadingTimingTrack.cpp
+++ b/Engine/Source/Developer/TraceInsights/Private/Insights/LoadingProfiler/Tracks/LoadingTimingTrack.cpp
@@ -25,6 +25,92 @@
 namespace UE::Insights::LoadingProfiler
 {
 
+////////////////////////////////////////////////////////////////////////////////////////////////////
+// FLoadingTimingTrackFilterContext
+////////////////////////////////////////////////////////////////////////////////////////////////////
+
+class FLoadingTimingTrackFilterContext : public FBaseFilterContext
+{
+public:
+	FLoadingTimingTrackFilterContext(const FLoadingTimingTrack& InTrack)
+		: Track(InTrack)
+	{
+	}
+	virtual ~FLoadingTimingTrackFilterContext() = default;
+
+	virtual bool HasData(int32 Key) const override
+	{
+		switch (Key)
+		{
+		case static_cast<int32>(EFilterField::StartTime):
+		case static_cast<int32>(EFilterField::EndTime):
+		case static_cast<int32>(EFilterField::Duration):
+		case static_cast<int32>(EFilterField::TrackName):
+			return true;
+		default:
+			return false;
+		}
+	}
+
+	virtual EFilterDataType GetDataType(int32 Key) const override
+	{
+		switch (Key)
+		{
+		case static_cast<int32>(EFilterField::StartTime):
+		case static_cast<int32>(EFilterField::EndTime):
+		case static_cast<int32>(EFilterField::Duration):
+			return EFilterDataType::Double;
+		case static_cast<int32>(EFilterField::TrackName):
+			return EFilterDataType::String;
+		default:
+			return EFilterDataType::Void;
+		}
+	}
+
+	virtual double GetDataAsDouble(int32 Key) const override
+	{
+		if (Key == static_cast<int32>(EFilterField::StartTime))
+		{
+			return EventStartTime;
+		}
+		if (Key == static_cast<int32>(EFilterField::EndTime))
+		{
+			return EventEndTime;
+		}
+		if (Key == static_cast<int32>(EFilterField::Duration))
+		{
+			return EventEndTime - EventStartTime;
+		}
+		return 0.0;
+	}
+
+	virtual const FString& GetDataAsString(int32 Key) const override
+	{
+		if (Key == static_cast<int32>(EFilterField::TrackName))
+		{
+			return Track.GetName();
+		}
+		return EmptyString;
+	}
+
+	void SetEvent(double StartTime, double EndTime, uint32 Depth, const TraceServices::FLoadTimeProfilerCpuEvent& InEvent)
+	{
+		EventStartTime = StartTime;
+		EventEndTime = EndTime;
+		EventDepth = Depth;
+		Event = &InEvent;
+	}
+
+private:
+	const FLoadingTimingTrack& Track;
+	double EventStartTime = 0.0;
+	double EventEndTime = 0.0;
+	uint32 EventDepth = 0;
+	const TraceServices::FLoadTimeProfilerCpuEvent* Event;
+};
+
+////////////////////////////////////////////////////////////////////////////////////////////////////
+// FLoadingTimingTrack
 ////////////////////////////////////////////////////////////////////////////////////////////////////
 
 INSIGHTS_IMPLEMENT_RTTI(FLoadingTimingTrack)
@@ -34,46 +120,55 @@ INSIGHTS_IMPLEMENT_RTTI(FLoadingTimingTrack)
 void FLoadingTimingTrack::BuildDrawState(ITimingEventsTrackDrawStateBuilder& Builder, const ITimingTrackUpdateContext& Context)
 {
 	TSharedPtr<const TraceServices::IAnalysisSession> Session = FInsightsManager::Get()->GetSession();
-	if (Session.IsValid() && TraceServices::ReadLoadTimeProfilerProvider(*Session.Get()))
+	if (!Session.IsValid())
 	{
-		TraceServices::FAnalysisSessionReadScope SessionReadScope(*Session.Get());
+		return;
+	}
 
-		const TraceServices::ILoadTimeProfilerProvider& LoadTimeProfilerProvider = *TraceServices::ReadLoadTimeProfilerProvider(*Session.Get());
+	if (const TraceServices::ILoadTimeProfilerProvider* LoadTimeProfilerProvider = TraceServices::ReadLoadTimeProfilerProvider(*Session.Get()))
+	{
+		TraceServices::FAnalysisSessionReadScope SessionReadScope(*Session.Get());
 
 		const FTimingTrackViewport& Viewport = Context.GetViewport();
 
-		LoadTimeProfilerProvider.ReadTimeline(TimelineIndex, [this, &Builder, &Viewport](const TraceServices::ILoadTimeProfilerProvider::CpuTimeline& Timeline)
-		{
-			if (FTimingEventsTrack::bUseDownSampling)
+		LoadTimeProfilerProvider->ReadTimeline(TimelineIndex,
+			[this, &Builder, &Viewport]
+			(const TraceServices::ILoadTimeProfilerProvider::CpuTimeline& Timeline)
 			{
-				const double SecondsPerPixel = 1.0 / Viewport.GetScaleX();
-				Timeline.EnumerateEventsDownSampled(Viewport.GetStartTime(), Viewport.GetEndTime(), SecondsPerPixel, [this, &Builder](double StartTime, double EndTime, uint32 Depth, const TraceServices::FLoadTimeProfilerCpuEvent& Event)
+				if (FTimingEventsTrack::bUseDownSampling)
 				{
-					if (Event.Package)
-					{
-						const TCHAR* Name = SharedState.GetEventName(Depth, Event);
-						const uint64 Type = static_cast<uint64>(Event.EventType);
-						const uint32 Color = 0;
-						Builder.AddEvent(StartTime, EndTime, Depth, Name, Type, Color);
-					}
-					return TraceServices::EEventEnumerate::Continue;
-				});
-			}
-			else
-			{
-				Timeline.EnumerateEvents(Viewport.GetStartTime(), Viewport.GetEndTime(), [this, &Builder](double StartTime, double EndTime, uint32 Depth, const TraceServices::FLoadTimeProfilerCpuEvent& Event)
+					const double SecondsPerPixel = 1.0 / Viewport.GetScaleX();
+					Timeline.EnumerateEventsDownSampled(Viewport.GetStartTime(), Viewport.GetEndTime(), SecondsPerPixel,
+						[this, &Builder]
+						(double StartTime, double EndTime, uint32 Depth, const TraceServices::FLoadTimeProfilerCpuEvent& Event)
+						{
+							if (Event.Package)
+							{
+								const TCHAR* Name = SharedState.GetEventName(Depth, Event);
+								const uint64 Type = static_cast<uint64>(Event.EventType);
+								const uint32 Color = 0;
+								Builder.AddEvent(StartTime, EndTime, Depth, Name, Type, Color);
+							}
+							return TraceServices::EEventEnumerate::Continue;
+						});
+				}
+				else
 				{
-					if (Event.Package)
-					{
-						const TCHAR* Name = SharedState.GetEventName(Depth, Event);
-						const uint64 Type = static_cast<uint64>(Event.EventType);
-						const uint32 Color = 0;
-						Builder.AddEvent(StartTime, EndTime, Depth, Name, Type, Color);
-					}
-					return TraceServices::EEventEnumerate::Continue;
-				});
-			}
-		});
+					Timeline.EnumerateEvents(Viewport.GetStartTime(), Viewport.GetEndTime(),
+						[this, &Builder]
+						(double StartTime, double EndTime, uint32 Depth, const TraceServices::FLoadTimeProfilerCpuEvent& Event)
+						{
+							if (Event.Package)
+							{
+								const TCHAR* Name = SharedState.GetEventName(Depth, Event);
+								const uint64 Type = static_cast<uint64>(Event.EventType);
+								const uint32 Color = 0;
+								Builder.AddEvent(StartTime, EndTime, Depth, Name, Type, Color);
+							}
+							return TraceServices::EEventEnumerate::Continue;
+						});
+				}
+			}); // ReadTimeline
 	}
 }
 
@@ -81,47 +176,42 @@ void FLoadingTimingTrack::BuildDrawState(ITimingEventsTrackDrawStateBuilder& Bui
 
 void FLoadingTimingTrack::BuildFilteredDrawState(ITimingEventsTrackDrawStateBuilder& Builder, const ITimingTrackUpdateContext& Context)
 {
+	TSharedPtr<const TraceServices::IAnalysisSession> Session = FInsightsManager::Get()->GetSession();
+	if (!Session.IsValid())
+	{
+		return;
+	}
+
 	if (HasCustomFilter())
 	{
-		using EFilterField = UE::Insights::EFilterField;
-		UE::Insights::FFilterContext FilterConfiguratorContext;
-		FilterConfiguratorContext.SetReturnValueForUnsetFilters(false);
-		FilterConfiguratorContext.AddFilterData<double>(static_cast<int32>(EFilterField::StartTime), 0.0f);
-		FilterConfiguratorContext.AddFilterData<double>(static_cast<int32>(EFilterField::EndTime), 0.0f);
-		FilterConfiguratorContext.AddFilterData<double>(static_cast<int32>(EFilterField::Duration), 0.0f);
-		FilterConfiguratorContext.AddFilterData<FString>(static_cast<int32>(EFilterField::TrackName), this->GetName());
-
-		TSharedPtr<const TraceServices::IAnalysisSession> Session = FInsightsManager::Get()->GetSession();
-		if (Session.IsValid() && TraceServices::ReadLoadTimeProfilerProvider(*Session.Get()))
+		if (const TraceServices::ILoadTimeProfilerProvider* LoadTimeProfilerProvider = TraceServices::ReadLoadTimeProfilerProvider(*Session.Get()))
 		{
-			TraceServices::FAnalysisSessionReadScope SessionReadScope(*Session.Get());
-
-			const TraceServices::ILoadTimeProfilerProvider& LoadTimeProfilerProvider = *TraceServices::ReadLoadTimeProfilerProvider(*Session.Get());
-
+			FLoadingTimingTrackFilterContext FilterContext(*this);
 			const FTimingTrackViewport& Viewport = Context.GetViewport();
 
-			LoadTimeProfilerProvider.ReadTimeline(TimelineIndex, [this, &Builder, &Viewport, &FilterConfiguratorContext](const TraceServices::ILoadTimeProfilerProvider::CpuTimeline& Timeline)
-			{
-				Timeline.EnumerateEvents(Viewport.GetStartTime(), Viewport.GetEndTime(), [this, &Builder, &FilterConfiguratorContext](double StartTime, double EndTime, uint32 Depth, const TraceServices::FLoadTimeProfilerCpuEvent& Event)
+			TraceServices::FAnalysisSessionReadScope SessionReadScope(*Session.Get());
+			LoadTimeProfilerProvider->ReadTimeline(TimelineIndex,
+				[this, &Builder, &Viewport, &FilterContext]
+				(const TraceServices::ILoadTimeProfilerProvider::CpuTimeline& Timeline)
 				{
-					if (Event.Package)
-					{
-						FilterConfiguratorContext.SetFilterData<double>(static_cast<int32>(EFilterField::StartTime), StartTime);
-						FilterConfiguratorContext.SetFilterData<double>(static_cast<int32>(EFilterField::EndTime), EndTime);
-						FilterConfiguratorContext.SetFilterData<double>(static_cast<int32>(EFilterField::Duration), EndTime - StartTime);
-
-						if (FilterConfigurator->ApplyFilters(FilterConfiguratorContext))
+					Timeline.EnumerateEvents(Viewport.GetStartTime(), Viewport.GetEndTime(),
+						[this, &Builder, &FilterContext]
+						(double StartTime, double EndTime, uint32 Depth, const TraceServices::FLoadTimeProfilerCpuEvent& Event)
 						{
-							const TCHAR* Name = SharedState.GetEventName(Depth, Event);
-							const uint64 Type = static_cast<uint64>(Event.EventType);
-							const uint32 Color = 0;
-							Builder.AddEvent(StartTime, EndTime, Depth, Name, Type, Color);
-						}
-					}
-
-					return TraceServices::EEventEnumerate::Continue;
-				});
-			});
+							if (Event.Package)
+							{
+								FilterContext.SetEvent(StartTime, EndTime, Depth, Event);
+								if (FilterConfigurator->ApplyFilters(FilterContext))
+								{
+									const TCHAR* Name = SharedState.GetEventName(Depth, Event);
+									const uint64 Type = static_cast<uint64>(Event.EventType);
+									const uint32 Color = 0;
+									Builder.AddEvent(StartTime, EndTime, Depth, Name, Type, Color);
+								}
+							}
+							return TraceServices::EEventEnumerate::Continue;
+						}); // EnumerateEvents
+				}); // ReadTimeline
 		}
 	}
 }
@@ -184,25 +274,29 @@ const TSharedPtr<const ITimingEvent> FLoadingTimingTrack::SearchEvent(const FTim
 {
 	TSharedPtr<const ITimingEvent> FoundEvent;
 
-	FindLoadTimeProfilerCpuEvent(InSearchParameters, [this, &FoundEvent](double InFoundStartTime, double InFoundEndTime, uint32 InFoundDepth, const TraceServices::FLoadTimeProfilerCpuEvent& InFoundEvent)
-	{
-		FoundEvent = MakeShared<FTimingEvent>(SharedThis(this), InFoundStartTime, InFoundEndTime, InFoundDepth);
-	});
+	FindLoadTimeProfilerCpuEvent(InSearchParameters,
+		[this, &FoundEvent]
+		(double InFoundStartTime, double InFoundEndTime, uint32 InFoundDepth, const TraceServices::FLoadTimeProfilerCpuEvent& InFoundEvent)
+		{
+			FoundEvent = MakeShared<FTimingEvent>(SharedThis(this), InFoundStartTime, InFoundEndTime, InFoundDepth);
+		});
 
 	return FoundEvent;
 }
 
 ////////////////////////////////////////////////////////////////////////////////////////////////////
 
-bool FLoadingTimingTrack::FindLoadTimeProfilerCpuEvent(const FTimingEventSearchParameters& InParameters, TFunctionRef<void(double, double, uint32, const TraceServices::FLoadTimeProfilerCpuEvent&)> InFoundPredicate) const
+bool FLoadingTimingTrack::FindLoadTimeProfilerCpuEvent(
+	const FTimingEventSearchParameters& InParameters,
+	TFunctionRef<void(double, double, uint32, const TraceServices::FLoadTimeProfilerCpuEvent&)> InFoundPredicate) const
 {
-	using EFilterField = UE::Insights::EFilterField;
-	UE::Insights::FFilterContext FilterConfiguratorContext;
-	FilterConfiguratorContext.SetReturnValueForUnsetFilters(false);
-	FilterConfiguratorContext.AddFilterData<double>(static_cast<int32>(EFilterField::StartTime), 0.0f);
-	FilterConfiguratorContext.AddFilterData<double>(static_cast<int32>(EFilterField::EndTime), 0.0f);
-	FilterConfiguratorContext.AddFilterData<double>(static_cast<int32>(EFilterField::Duration), 0.0f);
-	FilterConfiguratorContext.AddFilterData<FString>(static_cast<int32>(EFilterField::TrackName), this->GetName());
+	TSharedPtr<const TraceServices::IAnalysisSession> Session = FInsightsManager::Get()->GetSession();
+	if (!Session.IsValid())
+	{
+		return false;
+	}
+
+	FLoadingTimingTrackFilterContext FilterContext(*this);
 
 	return TTimingEventSearch<TraceServices::FLoadTimeProfilerCpuEvent>::Search(
 
@@ -215,45 +309,43 @@ bool FLoadingTimingTrack::FindLoadTimeProfilerCpuEvent(const FTimingEventSearchP
 			TSharedPtr<const TraceServices::IAnalysisSession> Session = FInsightsManager::Get()->GetSession();
 			if (Session.IsValid())
 			{
-				if (TraceServices::ReadLoadTimeProfilerProvider(*Session.Get()))
+				if (const TraceServices::ILoadTimeProfilerProvider* LoadTimeProfilerProvider = TraceServices::ReadLoadTimeProfilerProvider(*Session.Get()))
 				{
-					const TraceServices::ILoadTimeProfilerProvider& LoadTimeProfilerProvider = *TraceServices::ReadLoadTimeProfilerProvider(*Session.Get());
 					TraceServices::FAnalysisSessionReadScope SessionReadScope(*Session.Get());
 
-					LoadTimeProfilerProvider.ReadTimeline(TimelineIndex, [&InContext](const TraceServices::ILoadTimeProfilerProvider::CpuTimeline& Timeline)
-					{
-						auto Callback =
-							[&InContext]
-							(double EventStartTime, double EventEndTime, uint32 EventDepth, const TraceServices::FLoadTimeProfilerCpuEvent& Event)
-							{
-								if (Event.Package)
+					LoadTimeProfilerProvider->ReadTimeline(TimelineIndex,
+						[&InContext]
+						(const TraceServices::ILoadTimeProfilerProvider::CpuTimeline& Timeline)
+						{
+							auto Callback =
+								[&InContext]
+								(double EventStartTime, double EventEndTime, uint32 EventDepth, const TraceServices::FLoadTimeProfilerCpuEvent& Event)
 								{
+									if (!Event.Package)
+									{
+										return TraceServices::EEventEnumerate::Continue;
+									}
 									InContext.Check(EventStartTime, EventEndTime, EventDepth, Event);
 									return InContext.ShouldContinueSearching() ?
 										TraceServices::EEventEnumerate::Continue :
 										TraceServices::EEventEnumerate::Stop;
-								}
-								else
-								{
-									return TraceServices::EEventEnumerate::Continue;
-								}
-							};
+								};
 
-						if (InContext.GetParameters().SearchDirection == FTimingEventSearchParameters::ESearchDirection::Forward)
-						{
-							Timeline.EnumerateEvents(InContext.GetParameters().StartTime, InContext.GetParameters().EndTime, Callback);
-						}
-						else
-						{
-							Timeline.EnumerateEventsBackwards(InContext.GetParameters().EndTime, InContext.GetParameters().StartTime, Callback);
-						}
-					});
+							if (InContext.GetParameters().SearchDirection == FTimingEventSearchParameters::ESearchDirection::Forward)
+							{
+								Timeline.EnumerateEvents(InContext.GetParameters().StartTime, InContext.GetParameters().EndTime, Callback);
+							}
+							else
+							{
+								Timeline.EnumerateEventsBackwards(InContext.GetParameters().EndTime, InContext.GetParameters().StartTime, Callback);
+							}
+						}); // ReadTimeline
 				}
 			}
 		},
 
 		// Filter Predicate
-		[&FilterConfiguratorContext, &InParameters]
+		[&FilterContext, &InParameters]
 		(double EventStartTime, double EventEndTime, uint32 EventDepth, const TraceServices::FLoadTimeProfilerCpuEvent& Event)
 		{
 			if (!InParameters.FilterExecutor.IsValid())
@@ -261,22 +353,8 @@ bool FLoadingTimingTrack::FindLoadTimeProfilerCpuEvent(const FTimingEventSearchP
 				return true;
 			}
 
-			TSharedPtr<const TraceServices::IAnalysisSession> Session = FInsightsManager::Get()->GetSession();
-			if (Session.IsValid())
-			{
-				TraceServices::FAnalysisSessionReadScope SessionReadScope(*Session.Get());
-
-				if (TraceServices::ReadLoadTimeProfilerProvider(*Session.Get()))
-				{
-					FilterConfiguratorContext.SetFilterData<double>(static_cast<int32>(EFilterField::StartTime), EventStartTime);
-					FilterConfiguratorContext.SetFilterData<double>(static_cast<int32>(EFilterField::EndTime), EventEndTime);
-					FilterConfiguratorContext.SetFilterData<double>(static_cast<int32>(EFilterField::Duration), EventEndTime - EventStartTime);
-
-					return InParameters.FilterExecutor->ApplyFilters(FilterConfiguratorContext);
-				}
-			}
-
-			return false;
+			FilterContext.SetEvent(EventStartTime, EventEndTime, EventDepth, Event);
+			return InParameters.FilterExecutor->ApplyFilters(FilterContext);
 		},
 
 		// Found Predicate
@@ -297,7 +375,7 @@ bool FLoadingTimingTrack::HasCustomFilter() const
 
 ////////////////////////////////////////////////////////////////////////////////////////////////////
 
-void FLoadingTimingTrack::SetFilterConfigurator(TSharedPtr<UE::Insights::FFilterConfigurator> InFilterConfigurator)
+void FLoadingTimingTrack::SetFilterConfigurator(TSharedPtr<FFilterConfigurator> InFilterConfigurator)
 {
 	if (FilterConfigurator != InFilterConfigurator)
 	{
diff --git a/Engine/Source/Developer/TraceInsights/Private/Insights/TimingProfiler/Tracks/RegionsTimingTrack.cpp b/Engine/Source/Developer/TraceInsights/Private/Insights/TimingProfiler/Tracks/RegionsTimingTrack.cpp
--- a/Engine/Source/Developer/TraceInsights/Private/Insights/TimingProfiler/Tracks/RegionsTimingTrack.cpp
+++ b/Engine/Source/Developer/TraceInsights/Private/Insights/TimingProfiler/Tracks/RegionsTimingTrack.cpp
@@ -31,6 +31,110 @@
 namespace UE::Insights::TimingProfiler
 {
 
+////////////////////////////////////////////////////////////////////////////////////////////////////
+// FTimingRegionsTrackFilterContext
+////////////////////////////////////////////////////////////////////////////////////////////////////
+
+class FTimingRegionsTrackFilterContext : public FBaseFilterContext
+{
+public:
+	FTimingRegionsTrackFilterContext(const FTimingRegionsTrack& InTrack)
+		: Track(InTrack)
+	{
+	}
+	virtual ~FTimingRegionsTrackFilterContext() = default;
+
+	virtual bool HasData(int32 Key) const override
+	{
+		switch (Key)
+		{
+		case static_cast<int32>(EFilterField::StartTime):
+		case static_cast<int32>(EFilterField::EndTime):
+		case static_cast<int32>(EFilterField::Duration):
+		case static_cast<int32>(EFilterField::TrackName):
+		case static_cast<int32>(EFilterField::RegionName):
+		case static_cast<int32>(EFilterField::RegionCategory):
+			return true;
+		default:
+			return false;
+		}
+	}
+
+	virtual EFilterDataType GetDataType(int32 Key) const override
+	{
+		switch (Key)
+		{
+		case static_cast<int32>(EFilterField::StartTime):
+		case static_cast<int32>(EFilterField::EndTime):
+		case static_cast<int32>(EFilterField::Duration):
+			return EFilterDataType::Double;
+		case static_cast<int32>(EFilterField::TrackName):
+		case static_cast<int32>(EFilterField::RegionName):
+		case static_cast<int32>(EFilterField::RegionCategory):
+			return EFilterDataType::String;
+		default:
+			return EFilterDataType::Void;
+		}
+	}
+
+	virtual double GetDataAsDouble(int32 Key) const override
+	{
+		if (Region)
+		{
+			if (Key == static_cast<int32>(EFilterField::StartTime))
+			{
+				return Region->BeginTime;
+			}
+			if (Key == static_cast<int32>(EFilterField::EndTime))
+			{
+				return Region->EndTime;
+			}
+			if (Key == static_cast<int32>(EFilterField::Duration))
+			{
+				return Region->EndTime - Region->BeginTime;
+			}
+		}
+		return 0.0;
+	}
+
+	virtual const FString& GetDataAsString(int32 Key) const override
+	{
+		if (Key == static_cast<int32>(EFilterField::TrackName))
+		{
+			return Track.GetName();
+		}
+		if (Region)
+		{
+			if (Key == static_cast<int32>(EFilterField::RegionName))
+			{
+				check(Region->Timer);
+				RegionName = Region->Timer->Name;
+				return RegionName;
+			}
+			if (Key == static_cast<int32>(EFilterField::RegionCategory))
+			{
+				check(Region->Timer);
+				if (Region->Timer->Category)
+				{
+					RegionCategory = Region->Timer->Category->Name;
+					return RegionCategory;
+				}
+			}
+		}
+		return EmptyString;
+	}
+
+	void SetRegion(const TraceServices::FTimeRegion& InRegion)
+	{
+		Region = &InRegion;
+	}
+
+private:
+	const FTimingRegionsTrack& Track;
+	const TraceServices::FTimeRegion* Region = nullptr;
+	mutable FString RegionName;
+	mutable FString RegionCategory;
+};
 
 ////////////////////////////////////////////////////////////////////////////////////////////////////
 // FTimingRegionsTrack
@@ -38,10 +142,13 @@ namespace UE::Insights::TimingProfiler
 
 INSIGHTS_IMPLEMENT_RTTI(FTimingRegionsTrack)
 
+////////////////////////////////////////////////////////////////////////////////////////////////////
+
 FTimingRegionsTrack::~FTimingRegionsTrack()
 {
 }
 
+////////////////////////////////////////////////////////////////////////////////////////////////////
 
 void FTimingRegionsTrack::SetRegionsCategory(const TCHAR* InRegionsCategory)
 {
@@ -49,6 +156,8 @@ void FTimingRegionsTrack::SetRegionsCategory(const TCHAR* InRegionsCategory)
 	SetName(FString(TEXT("Timing Regions - ")) + InRegionsCategory);
 }
 
+////////////////////////////////////////////////////////////////////////////////////////////////////
+
 void FTimingRegionsTrack::BuildContextMenu(FMenuBuilder& MenuBuilder)
 {
 	FTimingEventsTrack::BuildContextMenu(MenuBuilder);
@@ -90,12 +199,15 @@ void FTimingRegionsTrack::InitTooltip(FTooltipDrawState& InOutTooltip, const ITi
 
 ////////////////////////////////////////////////////////////////////////////////////////////////////
 
-void FTimingRegionsTrack::BuildDrawState(
-	ITimingEventsTrackDrawStateBuilder& Builder,
-	const ITimingTrackUpdateContext& Context)
+void FTimingRegionsTrack::BuildDrawState(ITimingEventsTrackDrawStateBuilder& Builder, const ITimingTrackUpdateContext& Context)
 {
-	const FTimingTrackViewport& Viewport = Context.GetViewport();
 	TSharedPtr<const TraceServices::IAnalysisSession> Session = FInsightsManager::Get()->GetSession();
+	if (!Session.IsValid())
+	{
+		return;
+	}
+
+	const FTimingTrackViewport& Viewport = Context.GetViewport();
 
 	const TraceServices::IRegionProvider& RegionProvider = TraceServices::ReadRegionProvider(*Session);
 	TraceServices::FProviderReadScopeLock RegionProviderScopedLock(RegionProvider);
@@ -139,10 +251,14 @@ void FTimingRegionsTrack::BuildDrawState(
 
 ////////////////////////////////////////////////////////////////////////////////////////////////////
 
-void FTimingRegionsTrack::BuildFilteredDrawState(
-	ITimingEventsTrackDrawStateBuilder& Builder,
-	const ITimingTrackUpdateContext& Context)
+void FTimingRegionsTrack::BuildFilteredDrawState(ITimingEventsTrackDrawStateBuilder& Builder, const ITimingTrackUpdateContext& Context)
 {
+	TSharedPtr<const TraceServices::IAnalysisSession> Session = FInsightsManager::Get()->GetSession();
+	if (!Session.IsValid())
+	{
+		return;
+	}
+
 	const TSharedPtr<ITimingEventFilter> EventFilterPtr = Context.GetEventFilter();
 	if (EventFilterPtr.IsValid() && EventFilterPtr->FilterTrack(*this))
 	{
@@ -155,100 +271,31 @@ void FTimingRegionsTrack::BuildFilteredDrawState(
 			FilterEventType = reinterpret_cast<TCHAR*>(EventFilter.GetEventType());
 		}
 
-		TSharedPtr<const TraceServices::IAnalysisSession> Session = FInsightsManager::Get()->GetSession();
-		if (Session.IsValid())
+		if (bFilterOnlyByEventType)
 		{
 			const TraceServices::IRegionProvider& RegionProvider = TraceServices::ReadRegionProvider(*Session);
 			TraceServices::FProviderReadScopeLock RegionProviderScopedLock(RegionProvider);
 
 			const FTimingTrackViewport& Viewport = Context.GetViewport();
 
-			if (bFilterOnlyByEventType)
-			{
-				int32 CurDepth = 0;
-				const TraceServices::IRegionTimeline* Timeline = RegionProvider.GetTimelineForCategory(RegionsCategory);
-				check(Timeline)
-				Timeline->EnumerateLanes(
-					[this, Viewport, &CurDepth, &Builder, FilterEventType]
-					(const TraceServices::FRegionLane& Lane, const int32 Depth)
-					{
-						bool RegionHadEvents = false;
-						Lane.EnumerateRegions(Viewport.GetStartTime(), Viewport.GetEndTime(),
-							[&Builder, &RegionHadEvents, &CurDepth, FilterEventType]
-							(const TraceServices::FTimeRegion& Region) -> bool
-							{
-								check(Region.Timer);
-								RegionHadEvents = true;
-								if (Region.Timer->Name == FilterEventType)
-								{
-									Builder.AddEvent(Region.BeginTime, Region.EndTime, CurDepth, Region.Timer->Name);
-								}
-								return true;
-							});
-
-						if (RegionHadEvents)
-						{
-							CurDepth++;
-						}
-					});
-			}
-			else // generic filter
-			{
-				//TODO: if (EventFilterPtr->FilterEvent(TimingEvent))
-			}
-		}
-	}
-
-	if (HasCustomFilter())
-	{
-		if (!FilterConfigurator.IsValid())
-		{
-			return;
-		}
-
-		FFilterContext FilterContext;
-		FilterContext.SetReturnValueForUnsetFilters(false);
-
-		FilterContext.AddFilterData<double>(static_cast<int32>(EFilterField::StartTime), 0.0f);
-		FilterContext.AddFilterData<double>(static_cast<int32>(EFilterField::EndTime), 0.0f);
-		FilterContext.AddFilterData<double>(static_cast<int32>(EFilterField::Duration), 0.0f);
-		FilterContext.AddFilterData<FString>(static_cast<int32>(EFilterField::TrackName), this->GetName());
-		FilterContext.AddFilterData<FString>(static_cast<int32>(EFilterField::RegionName), "");
-		FilterContext.AddFilterData<FString>(static_cast<int32>(EFilterField::RegionCategory), "");
-
-		TSharedPtr<const TraceServices::IAnalysisSession> Session = FInsightsManager::Get()->GetSession();
-
-		if (Session.IsValid())
-		{
-			const TraceServices::IRegionProvider& RegionProvider = TraceServices::ReadRegionProvider(*Session);
-			TraceServices::FProviderReadScopeLock RegionProviderScopedLock(RegionProvider);
-			const FTimingTrackViewport& Viewport = Context.GetViewport();
-
 			int32 CurDepth = 0;
 			const TraceServices::IRegionTimeline* Timeline = RegionProvider.GetTimelineForCategory(RegionsCategory);
 			check(Timeline)
 			Timeline->EnumerateLanes(
-				[this, Viewport, &CurDepth, &Builder, &FilterContext]
+				[this, Viewport, &CurDepth, &Builder, FilterEventType]
 				(const TraceServices::FRegionLane& Lane, const int32 Depth)
 				{
 					bool RegionHadEvents = false;
 					Lane.EnumerateRegions(Viewport.GetStartTime(), Viewport.GetEndTime(),
-						[&Builder, &RegionHadEvents, &CurDepth, &FilterContext, this]
+						[&Builder, &RegionHadEvents, &CurDepth, FilterEventType]
 						(const TraceServices::FTimeRegion& Region) -> bool
 						{
 							check(Region.Timer);
-							FilterContext.SetFilterData<double>(static_cast<int32>(EFilterField::StartTime), Region.BeginTime);
-							FilterContext.SetFilterData<double>(static_cast<int32>(EFilterField::EndTime), Region.EndTime);
-							FilterContext.SetFilterData<double>(static_cast<int32>(EFilterField::Duration), Region.EndTime - Region.BeginTime);
-							FilterContext.SetFilterData<FString>(static_cast<int32>(EFilterField::RegionName), Region.Timer->Name);
-							FilterContext.SetFilterData<FString>(static_cast<int32>(EFilterField::RegionCategory), Region.Timer->Category ? Region.Timer->Category->Name : nullptr);
-
-							RegionHadEvents = true;
-							if (FilterConfigurator->ApplyFilters(FilterContext))
+							if (Region.Timer->Name == FilterEventType)
 							{
 								Builder.AddEvent(Region.BeginTime, Region.EndTime, CurDepth, Region.Timer->Name);
 							}
-
+							RegionHadEvents = true;
 							return true;
 						});
 
@@ -258,6 +305,47 @@ void FTimingRegionsTrack::BuildFilteredDrawState(
 					}
 				});
 		}
+		else // generic filter
+		{
+			//TODO: if (EventFilterPtr->FilterEvent(TimingEvent))
+		}
+	}
+
+	if (HasCustomFilter())
+	{
+		FTimingRegionsTrackFilterContext FilterContext(*this);
+		const FTimingTrackViewport& Viewport = Context.GetViewport();
+
+		const TraceServices::IRegionProvider& RegionProvider = TraceServices::ReadRegionProvider(*Session);
+		TraceServices::FProviderReadScopeLock RegionProviderScopedLock(RegionProvider);
+
+		int32 CurDepth = 0;
+		const TraceServices::IRegionTimeline* Timeline = RegionProvider.GetTimelineForCategory(RegionsCategory);
+		check(Timeline)
+		Timeline->EnumerateLanes(
+			[this, &Viewport, &CurDepth, &Builder, &FilterContext]
+			(const TraceServices::FRegionLane& Lane, const int32 Depth)
+			{
+				bool RegionHadEvents = false;
+				Lane.EnumerateRegions(Viewport.GetStartTime(), Viewport.GetEndTime(),
+					[&Builder, &RegionHadEvents, &CurDepth, &FilterContext, this]
+					(const TraceServices::FTimeRegion& Region) -> bool
+					{
+						check(Region.Timer);
+						FilterContext.SetRegion(Region);
+						if (FilterConfigurator->ApplyFilters(FilterContext))
+						{
+							Builder.AddEvent(Region.BeginTime, Region.EndTime, CurDepth, Region.Timer->Name);
+						}
+						RegionHadEvents = true;
+						return true;
+					});
+
+				if (RegionHadEvents)
+				{
+					CurDepth++;
+				}
+			});
 	}
 }
 
@@ -285,24 +373,23 @@ bool FTimingRegionsTrack::FindRegionEvent(
 	const FTimingEventSearchParameters& InParameters,
 	TFunctionRef<void(double, double, uint32, const TraceServices::FTimeRegion&)> InFoundPredicate) const
 {
+	TSharedPtr<const TraceServices::IAnalysisSession> Session = FInsightsManager::Get()->GetSession();
+	if (!Session.IsValid())
+	{
+		return false;
+	}
+
 	// If the query start time is larger than the end of the session return false.
 	{
-		TSharedPtr<const TraceServices::IAnalysisSession> Session = FInsightsManager::Get()->GetSession();
 		TraceServices::FAnalysisSessionReadScope SessionReadScope(*Session.Get());
-		if (Session.IsValid() && InParameters.StartTime > Session->GetDurationSeconds())
+		const double SessionDuration = Session->GetDurationSeconds();;
+		if (InParameters.StartTime > SessionDuration)
 		{
 			return false;
 		}
 	}
 
-	FFilterContext FilterConfiguratorContext;
-	FilterConfiguratorContext.SetReturnValueForUnsetFilters(false);
-	FilterConfiguratorContext.AddFilterData<double>(static_cast<int32>(EFilterField::StartTime), 0.0f);
-	FilterConfiguratorContext.AddFilterData<double>(static_cast<int32>(EFilterField::EndTime), 0.0f);
-	FilterConfiguratorContext.AddFilterData<double>(static_cast<int32>(EFilterField::Duration), 0.0f);
-	FilterConfiguratorContext.AddFilterData<FString>(static_cast<int32>(EFilterField::TrackName), this->GetName());
-	FilterConfiguratorContext.AddFilterData<FString>(static_cast<int32>(EFilterField::RegionName), "");
-	FilterConfiguratorContext.AddFilterData<FString>(static_cast<int32>(EFilterField::RegionCategory), "");
+	FTimingRegionsTrackFilterContext FilterContext(*this);
 
 	return TTimingEventSearch<TraceServices::FTimeRegion>::Search(
 
@@ -313,38 +400,35 @@ bool FTimingRegionsTrack::FindRegionEvent(
 		(TTimingEventSearch<TraceServices::FTimeRegion>::FContext& InContext)
 		{
 			TSharedPtr<const TraceServices::IAnalysisSession> Session = FInsightsManager::Get()->GetSession();
-			const TraceServices::IRegionProvider& RegionProvider = TraceServices::ReadRegionProvider(*Session);
-			TraceServices::FProviderReadScopeLock RegionProviderScopedLock(RegionProvider);
-
-			const TraceServices::IRegionTimeline* Timeline = RegionProvider.GetTimelineForCategory(RegionsCategory);
-			check(Timeline)
+			if (Session.IsValid())
+			{
+				const TraceServices::IRegionProvider& RegionProvider = TraceServices::ReadRegionProvider(*Session);
+				TraceServices::FProviderReadScopeLock RegionProviderScopedLock(RegionProvider);
 
-			auto Callback =
-				[&InContext]
-				(const TraceServices::FTimeRegion& Region)
-				{
-					InContext.Check(Region.BeginTime, Region.EndTime, Region.Depth, Region);
+				const TraceServices::IRegionTimeline* Timeline = RegionProvider.GetTimelineForCategory(RegionsCategory);
+				check(Timeline)
 
-					if (!InContext.ShouldContinueSearching())
+				auto Callback =
+					[&InContext]
+					(const TraceServices::FTimeRegion& Region) -> bool
 					{
-						return false;
-					}
+						InContext.Check(Region.BeginTime, Region.EndTime, Region.Depth, Region);
+						return InContext.ShouldContinueSearching();
+					};
 
-					return true;
-				};
-
-			if (InContext.GetParameters().SearchDirection == FTimingEventSearchParameters::ESearchDirection::Forward)
-			{
-				Timeline->EnumerateRegions(InContext.GetParameters().StartTime, InContext.GetParameters().EndTime, Callback);
-			}
-			else
-			{
-				Timeline->EnumerateRegionsBackwards(InContext.GetParameters().EndTime, InContext.GetParameters().StartTime, Callback);
+				if (InContext.GetParameters().SearchDirection == FTimingEventSearchParameters::ESearchDirection::Forward)
+				{
+					Timeline->EnumerateRegions(InContext.GetParameters().StartTime, InContext.GetParameters().EndTime, Callback);
+				}
+				else
+				{
+					Timeline->EnumerateRegionsBackwards(InContext.GetParameters().EndTime, InContext.GetParameters().StartTime, Callback);
+				}
 			}
 		},
 
 		// Filter Predicate
-		[&FilterConfiguratorContext, &InParameters]
+		[&FilterContext, &InParameters]
 		(double EventStartTime, double EventEndTime, uint32 EventDepth, const TraceServices::FTimeRegion& Region)
 		{
 			if (!InParameters.FilterExecutor.IsValid())
@@ -352,20 +436,9 @@ bool FTimingRegionsTrack::FindRegionEvent(
 				return true;
 			}
 
-			TSharedPtr<const TraceServices::IAnalysisSession> Session = FInsightsManager::Get()->GetSession();
-			if (!Session.IsValid())
-			{
-				// Region.Timer->Name and Region.Category strings are only available during a valid session.
-				return false;
-			}
-
 			check(Region.Timer);
-			FilterConfiguratorContext.SetFilterData<double>(static_cast<int32>(EFilterField::StartTime), EventStartTime);
-			FilterConfiguratorContext.SetFilterData<double>(static_cast<int32>(EFilterField::EndTime), EventEndTime);
-			FilterConfiguratorContext.SetFilterData<double>(static_cast<int32>(EFilterField::Duration), EventEndTime - EventStartTime);
-			FilterConfiguratorContext.SetFilterData<FString>(static_cast<int32>(EFilterField::RegionName), Region.Timer->Name);
-			FilterConfiguratorContext.SetFilterData<FString>(static_cast<int32>(EFilterField::RegionCategory), Region.Timer->Category ? Region.Timer->Category->Name : nullptr);
-			return InParameters.FilterExecutor->ApplyFilters(FilterConfiguratorContext);
+			FilterContext.SetRegion(Region);
+			return InParameters.FilterExecutor->ApplyFilters(FilterContext);
 		},
 
 		// Found Predicate
@@ -379,6 +452,13 @@ bool FTimingRegionsTrack::FindRegionEvent(
 
 ////////////////////////////////////////////////////////////////////////////////////////////////////
 
+bool FTimingRegionsTrack::HasCustomFilter() const
+{
+	return FilterConfigurator.IsValid() && !FilterConfigurator->IsEmpty();
+}
+
+////////////////////////////////////////////////////////////////////////////////////////////////////
+
 void FTimingRegionsTrack::SetFilterConfigurator(TSharedPtr<FFilterConfigurator> InFilterConfigurator)
 {
 	if (FilterConfigurator != InFilterConfigurator)
@@ -390,13 +470,6 @@ void FTimingRegionsTrack::SetFilterConfigurator(TSharedPtr<FFilterConfigurator>
 
 ////////////////////////////////////////////////////////////////////////////////////////////////////
 
-bool FTimingRegionsTrack::HasCustomFilter() const
-{
-	return FilterConfigurator.IsValid() && !FilterConfigurator->IsEmpty();
-}
-
-////////////////////////////////////////////////////////////////////////////////////////////////////
-
 void FTimingRegionsTrack::OnClipboardCopyEvent(const ITimingEvent& InSelectedEvent) const
 {
 	if (InSelectedEvent.CheckTrack(this) && InSelectedEvent.Is<FTimingEvent>())
diff --git a/Engine/Source/Developer/TraceInsights/Private/Insights/TimingProfiler/Tracks/RegionsTimingTrack.h b/Engine/Source/Developer/TraceInsights/Private/Insights/TimingProfiler/Tracks/RegionsTimingTrack.h
--- a/Engine/Source/Developer/TraceInsights/Private/Insights/TimingProfiler/Tracks/RegionsTimingTrack.h
+++ b/Engine/Source/Developer/TraceInsights/Private/Insights/TimingProfiler/Tracks/RegionsTimingTrack.h
@@ -44,7 +44,8 @@ public:
 	virtual void OnClipboardCopyEvent(const ITimingEvent& InSelectedEvent) const override;
 
 	void SetRegionsCategory(const TCHAR* InRegionsCategory);
-protected:
+
+private:
 	bool FindRegionEvent(const FTimingEventSearchParameters& InParameters, TFunctionRef<void(double, double, uint32, const TraceServices::FTimeRegion&)> InFoundPredicate) const;
 
 private:
diff --git a/Engine/Source/Developer/TraceInsights/Private/Insights/TimingProfiler/Tracks/ThreadTimingTrack.cpp b/Engine/Source/Developer/TraceInsights/Private/Insights/TimingProfiler/Tracks/ThreadTimingTrack.cpp
--- a/Engine/Source/Developer/TraceInsights/Private/Insights/TimingProfiler/Tracks/ThreadTimingTrack.cpp
+++ b/Engine/Source/Developer/TraceInsights/Private/Insights/TimingProfiler/Tracks/ThreadTimingTrack.cpp
@@ -37,12 +37,124 @@
 namespace UE::Insights::TimingProfiler
 {
 
-INSIGHTS_IMPLEMENT_RTTI(FThreadTimingTrack);
+////////////////////////////////////////////////////////////////////////////////////////////////////
+// FThreadTimingTrackFilterContext
+////////////////////////////////////////////////////////////////////////////////////////////////////
+
+class FThreadTimingTrackFilterContext : public FBaseFilterContext
+{
+public:
+	FThreadTimingTrackFilterContext(const FThreadTimingTrackImpl& InTrack)
+		: Track(InTrack)
+	{
+	}
+	virtual ~FThreadTimingTrackFilterContext() = default;
+
+	virtual bool HasData(int32 Key) const override
+	{
+		switch (Key)
+		{
+		case static_cast<int32>(EFilterField::StartTime):
+		case static_cast<int32>(EFilterField::EndTime):
+		case static_cast<int32>(EFilterField::Duration):
+		case static_cast<int32>(EFilterField::TrackName):
+		case static_cast<int32>(EFilterField::TimerId):
+		case static_cast<int32>(EFilterField::TimerName):
+		case static_cast<int32>(EFilterField::Metadata):
+			return true;
+		default:
+			return false;
+		}
+	}
+
+	virtual EFilterDataType GetDataType(int32 Key) const override
+	{
+		switch (Key)
+		{
+		case static_cast<int32>(EFilterField::StartTime):
+		case static_cast<int32>(EFilterField::EndTime):
+		case static_cast<int32>(EFilterField::Duration):
+			return EFilterDataType::Double;
+		case static_cast<int32>(EFilterField::TrackName):
+			return EFilterDataType::String;
+		case static_cast<int32>(EFilterField::TimerId):
+		case static_cast<int32>(EFilterField::TimerName): // exposed as TimerId
+		case static_cast<int32>(EFilterField::Metadata):
+			return EFilterDataType::Int64;
+		default:
+			return EFilterDataType::Void;
+		}
+	}
+
+	virtual int64 GetDataAsInt64(int32 Key) const override
+	{
+		if (Key == static_cast<int32>(EFilterField::TimerId) ||
+			Key == static_cast<int32>(EFilterField::TimerName)) // exposed as TimerId
+		{
+			check(TimerReader);
+			const TraceServices::FTimingProfilerTimer* Timer = TimerReader->GetTimer(TimerIndex);
+			return ensure(Timer != nullptr) ? int64(Timer->Id) : int64(~0u);
+		}
+		if (Key == static_cast<int32>(EFilterField::Metadata))
+		{
+			return int64(TimerIndex);
+		}
+		return 0;
+	}
+
+	virtual double GetDataAsDouble(int32 Key) const override
+	{
+		if (Key == static_cast<int32>(EFilterField::StartTime))
+		{
+			return EventStartTime;
+		}
+		if (Key == static_cast<int32>(EFilterField::EndTime))
+		{
+			return EventEndTime;
+		}
+		if (Key == static_cast<int32>(EFilterField::Duration))
+		{
+			return EventEndTime - EventStartTime;
+		}
+		return 0.0;
+	}
+
+	virtual const FString& GetDataAsString(int32 Key) const override
+	{
+		if (Key == static_cast<int32>(EFilterField::TrackName))
+		{
+			return Track.GetName();
+		}
+		return EmptyString;
+	}
+
+	void SetTimerReader(const TraceServices::ITimingProfilerTimerReader& InTimerReader)
+	{
+		TimerReader = &InTimerReader;
+	}
+
+	void SetEvent(double StartTime, double EndTime, uint32 Depth, const TraceServices::FTimingProfilerEvent& Event)
+	{
+		EventStartTime = StartTime;
+		EventEndTime = EndTime;
+		EventDepth = Depth;
+		TimerIndex = Event.TimerIndex;
+	}
+
+private:
+	const FThreadTimingTrackImpl& Track;
+	const TraceServices::ITimingProfilerTimerReader* TimerReader;
+	double EventStartTime = 0.0;
+	double EventEndTime = 0.0;
+	uint32 EventDepth = 0;
+	uint32 TimerIndex = 0;
+};
 
 ////////////////////////////////////////////////////////////////////////////////////////////////////
 // FThreadTimingTrackImpl
 ////////////////////////////////////////////////////////////////////////////////////////////////////
 
+INSIGHTS_IMPLEMENT_RTTI(FThreadTimingTrack);
 INSIGHTS_IMPLEMENT_RTTI(FThreadTimingTrackImpl)
 
 ////////////////////////////////////////////////////////////////////////////////////////////////////
@@ -775,7 +887,7 @@ void FThreadTimingTrackImpl::BuildFilteredDrawState(ITimingEventsTrackDrawStateB
 					(const TraceServices::ITimingProfilerProvider::Timeline& Timeline)
 					{
 						TArray<TArray<FPendingEventInfo>> FilteredEvents;
-						TArray<FFilterContext> FilterContexts;
+						TArray<FThreadTimingTrackFilterContext> FilterContexts;
 
 						TraceServices::ITimeline<TraceServices::FTimingProfilerEvent>::EnumerateAsyncParams Params;
 						constexpr uint32 LargeTimelineThreshold = 50 * 1000 * 1000;
@@ -786,54 +898,38 @@ void FThreadTimingTrackImpl::BuildFilteredDrawState(ITimingEventsTrackDrawStateB
 								Params.MaxOccupancy = 0.75f; // This filter can be slow so reduce occupancy to avoid starvation.
 							}
 						}
+
 						Params.IntervalStart = Viewport.GetStartTime();
 						Params.IntervalEnd = Viewport.GetEndTime();
 
 						// Note: Enumerating events for filtering should not use downsampling.
 						Params.Resolution = 0.0;
 						Params.SetupCallback =
-							[this, &FilteredEvents, &FilterContexts]
+							[this, &FilteredEvents, &FilterContexts, &TimerReader]
 							(uint32 NumTasks)
 							{
 								FilteredEvents.AddDefaulted(NumTasks);
-								FilterContexts.AddDefaulted(NumTasks);
-								for (FFilterContext& Context : FilterContexts)
+								FilterContexts.Reserve(NumTasks);
+								for (uint32 TaskIndex = 0; TaskIndex < NumTasks; ++TaskIndex)
 								{
-									Context.SetReturnValueForUnsetFilters(false);
-									Context.AddFilterData<double>(static_cast<int32>(EFilterField::StartTime), 0.0f);
-									Context.AddFilterData<double>(static_cast<int32>(EFilterField::EndTime), 0.0f);
-									Context.AddFilterData<double>(static_cast<int32>(EFilterField::Duration), 0.0f);
-									Context.AddFilterData<FString>(static_cast<int32>(EFilterField::TrackName), this->GetName());
-									Context.AddFilterData<int64>(static_cast<int32>(EFilterField::TimerId), 0);
-									Context.AddFilterData<int64>(static_cast<int32>(EFilterField::TimerName), 0);
-									Context.AddFilterData<int64>(static_cast<int32>(EFilterField::Metadata), 0);
+									FThreadTimingTrackFilterContext& FilterContext = FilterContexts.Add_GetRef(FThreadTimingTrackFilterContext(*this));
+									FilterContext.SetTimerReader(TimerReader);
 								}
 							};
 						Params.EventRangeCallback =
-							[this, &Builder, &TimerReader, &FilteredEvents, &FilterContexts]
+							[this, &Builder, &FilteredEvents, &FilterContexts]
 							(double StartTime, double EndTime, uint32 Depth, const TraceServices::FTimingProfilerEvent& Event, uint32 TaskIndex)
 							{
-								const TraceServices::FTimingProfilerTimer* Timer = TimerReader.GetTimer(Event.TimerIndex);
-								if (ensure(Timer != nullptr))
+								FThreadTimingTrackFilterContext& FilterContext = FilterContexts[TaskIndex];
+								FilterContext.SetEvent(StartTime, EndTime, Depth, Event);
+								if (FilterConfigurator->ApplyFilters(FilterContext))
 								{
-									FFilterContext& Context = FilterContexts[TaskIndex];
-									Context.SetFilterData<double>(static_cast<int32>(EFilterField::StartTime), StartTime);
-									Context.SetFilterData<double>(static_cast<int32>(EFilterField::EndTime), EndTime);
-									Context.SetFilterData<double>(static_cast<int32>(EFilterField::Duration), EndTime - StartTime);
-									// The TimerName filter also translates to the numeric Id for performance reasons.
-									Context.SetFilterData<int64>(static_cast<int32>(EFilterField::TimerId), Timer->Id);
-									Context.SetFilterData<int64>(static_cast<int32>(EFilterField::TimerName), Timer->Id);
-									Context.SetFilterData<int64>(static_cast<int32>(EFilterField::Metadata), Event.TimerIndex);
-
-									if (FilterConfigurator->ApplyFilters(Context))
-									{
-										FPendingEventInfo TimelineEvent;
-										TimelineEvent.StartTime = StartTime;
-										TimelineEvent.EndTime = EndTime;
-										TimelineEvent.Depth = Depth;
-										TimelineEvent.TimerIndex = Event.TimerIndex;
-										FilteredEvents[TaskIndex].Add(TimelineEvent);
-									}
+									FPendingEventInfo TimelineEvent;
+									TimelineEvent.StartTime = StartTime;
+									TimelineEvent.EndTime = EndTime;
+									TimelineEvent.Depth = Depth;
+									TimelineEvent.TimerIndex = Event.TimerIndex;
+									FilteredEvents[TaskIndex].Add(TimelineEvent);
 								}
 								return TraceServices::EEventEnumerate::Continue;
 							};
@@ -1149,18 +1245,21 @@ const TSharedPtr<const ITimingEvent> FThreadTimingTrackImpl::GetEvent(float InPo
 const TSharedPtr<const ITimingEvent> FThreadTimingTrackImpl::SearchEvent(const FTimingEventSearchParameters& InSearchParameters) const
 {
 	TSharedPtr<FThreadTrackEvent> FoundEvent;
-	FindTimingProfilerEvent(InSearchParameters, [this, &FoundEvent](double InFoundStartTime, double InFoundEndTime, uint32 InFoundDepth, const TraceServices::FTimingProfilerEvent& InFoundEvent)
-	{
-		FoundEvent = MakeShared<FThreadTrackEvent>(SharedThis(this), InFoundStartTime, InFoundEndTime, InFoundDepth);
-		FoundEvent->SetTimerIndex(InFoundEvent.TimerIndex);
 
-		uint32 TimerId = 0;
-		bool bIsValidTimerId = TimerIndexToTimerId(InFoundEvent.TimerIndex, TimerId);
-		if (bIsValidTimerId)
+	FindTimingProfilerEvent(InSearchParameters,
+		[this, &FoundEvent]
+		(double InFoundStartTime, double InFoundEndTime, uint32 InFoundDepth, const TraceServices::FTimingProfilerEvent& InFoundEvent)
 		{
-			FoundEvent->SetTimerId(TimerId);
-		}
-	});
+			FoundEvent = MakeShared<FThreadTrackEvent>(SharedThis(this), InFoundStartTime, InFoundEndTime, InFoundDepth);
+			FoundEvent->SetTimerIndex(InFoundEvent.TimerIndex);
+
+			uint32 TimerId = 0;
+			bool bIsValidTimerId = TimerIndexToTimerId(InFoundEvent.TimerIndex, TimerId);
+			if (bIsValidTimerId)
+			{
+				FoundEvent->SetTimerId(TimerId);
+			}
+		});
 
 	return FoundEvent;
 }
@@ -1294,14 +1393,18 @@ void FThreadTimingTrackImpl::BuildContextMenu(FMenuBuilder& MenuBuilder)
 
 ////////////////////////////////////////////////////////////////////////////////////////////////////
 
-bool FThreadTimingTrackImpl::FindTimingProfilerEvent(const FThreadTrackEvent& InTimingEvent, TFunctionRef<void(double, double, uint32, const TraceServices::FTimingProfilerEvent&)> InFoundPredicate) const
+bool FThreadTimingTrackImpl::FindTimingProfilerEvent(
+	const FThreadTrackEvent& InTimingEvent,
+	TFunctionRef<void(double, double, uint32, const TraceServices::FTimingProfilerEvent&)> InFoundPredicate) const
 {
-	auto MatchEvent = [&InTimingEvent](double InStartTime, double InEndTime, uint32 InDepth)
-	{
-		return InDepth == InTimingEvent.GetDepth()
-			&& InStartTime == InTimingEvent.GetStartTime()
-			&& InEndTime == InTimingEvent.GetEndTime();
-	};
+	auto MatchEvent =
+		[&InTimingEvent]
+		(double InStartTime, double InEndTime, uint32 InDepth)
+		{
+			return InDepth == InTimingEvent.GetDepth()
+				&& InStartTime == InTimingEvent.GetStartTime()
+				&& InEndTime == InTimingEvent.GetEndTime();
+		};
 
 	const double Time = (InTimingEvent.GetStartTime() + InTimingEvent.GetEndTime()) / 2;
 	FTimingEventSearchParameters SearchParameters(Time, Time, ETimingEventSearchFlags::StopAtFirstMatch, MatchEvent);
@@ -1311,50 +1414,62 @@ bool FThreadTimingTrackImpl::FindTimingProfilerEvent(const FThreadTrackEvent& In
 
 ////////////////////////////////////////////////////////////////////////////////////////////////////
 
-bool FThreadTimingTrackImpl::FindTimingProfilerEvent(const FTimingEventSearchParameters& InParameters, TFunctionRef<void(double, double, uint32, const TraceServices::FTimingProfilerEvent&)> InFoundPredicate) const
+bool FThreadTimingTrackImpl::FindTimingProfilerEvent(
+	const FTimingEventSearchParameters& InParameters,
+	TFunctionRef<void(double, double, uint32, const TraceServices::FTimingProfilerEvent&)> InFoundPredicate) const
 {
-	FFilterContext FilterConfiguratorContext;
-	FilterConfiguratorContext.SetReturnValueForUnsetFilters(false);
-	FilterConfiguratorContext.AddFilterData<double>(static_cast<int32>(EFilterField::StartTime), 0.0f);
-	FilterConfiguratorContext.AddFilterData<double>(static_cast<int32>(EFilterField::EndTime), 0.0f);
-	FilterConfiguratorContext.AddFilterData<double>(static_cast<int32>(EFilterField::Duration), 0.0f);
-	FilterConfiguratorContext.AddFilterData<FString>(static_cast<int32>(EFilterField::TrackName), this->GetName());
-	FilterConfiguratorContext.AddFilterData<int64>(static_cast<int32>(EFilterField::TimerId), 0);
-	FilterConfiguratorContext.AddFilterData<int64>(static_cast<int32>(EFilterField::TimerName), 0);
-	FilterConfiguratorContext.AddFilterData<int64>(static_cast<int32>(EFilterField::Metadata), 0);
+	TSharedPtr<const TraceServices::IAnalysisSession> Session = FInsightsManager::Get()->GetSession();
+	if (!Session.IsValid())
+	{
+		return false;
+	}
+
+	FThreadTimingTrackFilterContext FilterContext(*this);
 
 	return TTimingEventSearch<TraceServices::FTimingProfilerEvent>::Search(
 
 		InParameters,
 
 		// Search Predicate
-		[this]
+		[this, &FilterContext]
 		(TTimingEventSearch<TraceServices::FTimingProfilerEvent>::FContext& InContext)
 		{
-			ReadTimeline(
-				[&InContext]
-				(const TraceServices::ITimingProfilerProvider::Timeline& Timeline)
+			ReadTimers(
+				[this, &InContext, &FilterContext]
+				(const TraceServices::ITimingProfilerProvider& TimingProfilerProvider, const TraceServices::ITimingProfilerTimerReader& TimerReader)
 				{
-					auto Callback = [&InContext](double EventStartTime, double EventEndTime, uint32 EventDepth, const TraceServices::FTimingProfilerEvent& Event)
-					{
-						InContext.Check(EventStartTime, EventEndTime, EventDepth, Event);
-						return InContext.ShouldContinueSearching() ? TraceServices::EEventEnumerate::Continue : TraceServices::EEventEnumerate::Stop;
-					};
+					FilterContext.SetTimerReader(TimerReader);
 
-					if (InContext.GetParameters().SearchDirection == FTimingEventSearchParameters::ESearchDirection::Forward)
-					{
-						Timeline.EnumerateEvents(InContext.GetParameters().StartTime, InContext.GetParameters().EndTime, Callback);
-					}
-					else
-					{
-						Timeline.EnumerateEventsBackwards(InContext.GetParameters().EndTime, InContext.GetParameters().StartTime, Callback);
-					}
+					ReadTimeline(
+						[&InContext]
+						(const TraceServices::ITimingProfilerProvider::Timeline& Timeline)
+						{
+							auto Callback =
+								[&InContext]
+								(double EventStartTime, double EventEndTime, uint32 EventDepth, const TraceServices::FTimingProfilerEvent& Event)
+								{
+									InContext.Check(EventStartTime, EventEndTime, EventDepth, Event);
+									return InContext.ShouldContinueSearching() ?
+										TraceServices::EEventEnumerate::Continue :
+										TraceServices::EEventEnumerate::Stop;
+								};
+
+							if (InContext.GetParameters().SearchDirection == FTimingEventSearchParameters::ESearchDirection::Forward)
+							{
+								Timeline.EnumerateEvents(InContext.GetParameters().StartTime, InContext.GetParameters().EndTime, Callback);
+							}
+							else
+							{
+								Timeline.EnumerateEventsBackwards(InContext.GetParameters().EndTime, InContext.GetParameters().StartTime, Callback);
+							}
+						}
+					); // ReadTimeline
 				}
-			); // ReadTimeline
+			); // ReadTimers
 		},
 
 		// Filter Predicate
-		[this, &FilterConfiguratorContext, &InParameters]
+		[&FilterContext, &InParameters]
 		(double EventStartTime, double EventEndTime, uint32 EventDepth, const TraceServices::FTimingProfilerEvent& Event)
 		{
 			if (!InParameters.FilterExecutor.IsValid())
@@ -1362,27 +1477,8 @@ bool FThreadTimingTrackImpl::FindTimingProfilerEvent(const FTimingEventSearchPar
 				return true;
 			}
 
-			FilterConfiguratorContext.SetFilterData<double>(static_cast<int32>(EFilterField::StartTime), EventStartTime);
-			FilterConfiguratorContext.SetFilterData<double>(static_cast<int32>(EFilterField::EndTime), EventEndTime);
-			FilterConfiguratorContext.SetFilterData<double>(static_cast<int32>(EFilterField::Duration), EventEndTime - EventStartTime);
-
-			uint32 TimerId = ~0u;
-			ReadTimers(
-				[&Event, &TimerId]
-				(const TraceServices::ITimingProfilerProvider& TimingProfilerProvider, const TraceServices::ITimingProfilerTimerReader& TimerReader)
-				{
-					const TraceServices::FTimingProfilerTimer* Timer = TimerReader.GetTimer(Event.TimerIndex);
-					if (ensure(Timer != nullptr))
-					{
-						TimerId = Timer->Id;
-					}
-				});
-			FilterConfiguratorContext.SetFilterData<int64>(static_cast<int32>(EFilterField::TimerId), TimerId);
-			FilterConfiguratorContext.SetFilterData<int64>(static_cast<int32>(EFilterField::TimerName), TimerId);
-
-			FilterConfiguratorContext.SetFilterData<int64>(static_cast<int32>(EFilterField::Metadata), Event.TimerIndex);
-
-			return InParameters.FilterExecutor->ApplyFilters(FilterConfiguratorContext);
+			FilterContext.SetEvent(EventStartTime, EventEndTime, EventDepth, Event);
+			return InParameters.FilterExecutor->ApplyFilters(FilterContext);
 		},
 
 		// Found Predicate
@@ -1442,13 +1538,6 @@ bool FThreadTimingTrackImpl::TimerIndexToTimerId(uint32 InTimerIndex, uint32& Ou
 
 ////////////////////////////////////////////////////////////////////////////////////////////////////
 
-bool FThreadTimingTrackImpl::HasCustomFilter() const
-{
-	return FilterConfigurator.IsValid() && !FilterConfigurator->IsEmpty();
-}
-
-////////////////////////////////////////////////////////////////////////////////////////////////////
-
 int32 FThreadTimingTrackImpl::GetDepthAt(double Time) const
 {
 	int32 Depth = 0;
@@ -1466,6 +1555,13 @@ int32 FThreadTimingTrackImpl::GetDepthAt(double Time) const
 
 ////////////////////////////////////////////////////////////////////////////////////////////////////
 
+bool FThreadTimingTrackImpl::HasCustomFilter() const
+{
+	return FilterConfigurator.IsValid() && !FilterConfigurator->IsEmpty();
+}
+
+////////////////////////////////////////////////////////////////////////////////////////////////////
+
 void FThreadTimingTrackImpl::SetFilterConfigurator(TSharedPtr<FFilterConfigurator> InFilterConfigurator)
 {
 	if (FilterConfigurator != InFilterConfigurator)
diff --git a/Engine/Source/Developer/TraceInsights/Private/Insights/TimingProfiler/ViewModels/TimerFilters.cpp b/Engine/Source/Developer/TraceInsights/Private/Insights/TimingProfiler/ViewModels/TimerFilters.cpp
--- a/Engine/Source/Developer/TraceInsights/Private/Insights/TimingProfiler/ViewModels/TimerFilters.cpp
+++ b/Engine/Source/Developer/TraceInsights/Private/Insights/TimingProfiler/ViewModels/TimerFilters.cpp
@@ -67,17 +67,15 @@ void FTimerNameFilterState::Update()
 
 ////////////////////////////////////////////////////////////////////////////////////////////////////
 
-bool FTimerNameFilterState::ApplyFilter(const FFilterContext& Context) const
+bool FTimerNameFilterState::ApplyFilter(const FBaseFilterContext& Context) const
 {
-	if (!Context.HasFilterData(static_cast<int32>(Filter->GetKey())))
+	if (!Context.HasData(Filter->GetKey()))
 	{
-		return Context.GetReturnValueForUnsetFilters();
+		return Context.GetDefaultReturnValue();
 	}
 
-	int64 Value;
-	Context.GetFilterData<int64>(Filter->GetKey(), Value);
-
-	return TimerIds.Contains(static_cast<uint32>(Value));
+	uint32 TimerId = static_cast<uint32>(Context.GetDataAsInt64(Filter->GetKey()));
+	return TimerIds.Contains(TimerId);
 }
 
 ////////////////////////////////////////////////////////////////////////////////////////////////////
@@ -98,7 +96,6 @@ bool FTimerNameFilterState::Equals(const FFilterState& Other) const
 TSharedRef<FFilterState> FTimerNameFilterState::DeepCopy() const
 {
 	TSharedRef<FTimerNameFilterState> Copy = MakeShared<FTimerNameFilterState>(*this);
-
 	return Copy;
 }
 
@@ -296,19 +293,18 @@ void FMetadataFilterState::Update()
 
 ////////////////////////////////////////////////////////////////////////////////////////////////////
 
-bool FMetadataFilterState::ApplyFilter(const FFilterContext& Context) const
+bool FMetadataFilterState::ApplyFilter(const FBaseFilterContext& Context) const
 {
-	if (!Context.HasFilterData(static_cast<int32>(Filter->GetKey())))
+	if (!Context.HasData(Filter->GetKey()))
 	{
-		return Context.GetReturnValueForUnsetFilters();
+		return Context.GetDefaultReturnValue();
 	}
 
-	int64 TimerIndex;
-	Context.GetFilterData<int64>(Filter->GetKey(), TimerIndex);
+	uint32 TimerIndex = static_cast<uint32>(Context.GetDataAsInt64(Filter->GetKey()));
 
 	TSharedPtr<const TraceServices::IAnalysisSession> Session = FInsightsManager::Get()->GetSession();
 	TraceServices::FAnalysisSessionReadScope SessionReadScope(*Session.Get());
-	TArrayView<const uint8> Metadata = TimerReader->GetMetadata(static_cast<uint32>(TimerIndex));
+	TArrayView<const uint8> Metadata = TimerReader->GetMetadata(TimerIndex);
 
 	if (Metadata.Num() > 0)
 	{
diff --git a/Engine/Source/Developer/TraceInsights/Private/Insights/TimingProfiler/ViewModels/TimerFilters.h b/Engine/Source/Developer/TraceInsights/Private/Insights/TimingProfiler/ViewModels/TimerFilters.h
--- a/Engine/Source/Developer/TraceInsights/Private/Insights/TimingProfiler/ViewModels/TimerFilters.h
+++ b/Engine/Source/Developer/TraceInsights/Private/Insights/TimingProfiler/ViewModels/TimerFilters.h
@@ -40,7 +40,7 @@ public:
 
 	virtual void Update() override;
 
-	virtual bool ApplyFilter(const FFilterContext& Context) const override;
+	virtual bool ApplyFilter(const FBaseFilterContext& Context) const override;
 
 	virtual void SetFilterValue(FString InFilterValue) override { FilterValue = InFilterValue; }
 
@@ -64,8 +64,8 @@ public:
 	virtual ~FTimerNameFilter() = default;
 
 	virtual TSharedRef<FFilterState> BuildFilterState()
-	{ 
-		return MakeShared<FTimerNameFilterState>(SharedThis(this)); 
+	{
+		return MakeShared<FTimerNameFilterState>(SharedThis(this));
 	}
 
 	virtual TSharedRef<FFilterState> BuildFilterState(const FFilterState& Other)
@@ -113,7 +113,7 @@ public:
 
 	virtual void Update() override;
 
-	virtual bool ApplyFilter(const FFilterContext& Context) const override;
+	virtual bool ApplyFilter(const FBaseFilterContext& Context) const override;
 
 	virtual void SetFilterValue(FString InFilterValue) override {}
 
@@ -148,10 +148,10 @@ private:
 	ConvertedDataVariant ConvertedData;
 
 	bool bShowAllMetadataEvents = false;
-	
+
 	TArray<TSharedPtr<FMetadataFilterDataTypeEntry>> AvailableDataTypes;
 	TSharedPtr<FMetadataFilterDataTypeEntry> SelectedDataType;
-	
+
 	TArray<TSharedPtr<IFilterOperator>> AvailableOperators;
 	TArray<TSharedPtr<IFilterOperator>> BoolOperators;
 
diff --git a/Engine/Source/Developer/TraceInsights/Private/Insights/Widgets/STimingView.cpp b/Engine/Source/Developer/TraceInsights/Private/Insights/Widgets/STimingView.cpp
--- a/Engine/Source/Developer/TraceInsights/Private/Insights/Widgets/STimingView.cpp
+++ b/Engine/Source/Developer/TraceInsights/Private/Insights/Widgets/STimingView.cpp
@@ -6074,11 +6074,57 @@ void STimingView::PopulateRegionCategorySuggestionList(const FString& Text, TArr
 
 ////////////////////////////////////////////////////////////////////////////////////////////////////
 
+class FBaseTimingTrackFilterContext : public FBaseFilterContext
+{
+public:
+	static constexpr int32 TrackNameKey = static_cast<int32>(UE::Insights::EFilterField::TrackName);
+
+public:
+	FBaseTimingTrackFilterContext() = default;
+	virtual ~FBaseTimingTrackFilterContext() = default;
+
+	virtual bool GetDefaultReturnValue() const override
+	{
+		return true;
+	}
+
+	virtual bool HasData(int32 Key) const override
+	{
+		return Key == TrackNameKey;
+	}
+
+	virtual EFilterDataType GetDataType(int32 Key) const override
+	{
+		return Key == TrackNameKey ? EFilterDataType::String : EFilterDataType::Void;
+	}
+
+	virtual FStringView GetDataAsStringView(int32 Key) const override
+	{
+		return Track && Key == TrackNameKey ? Track->GetName() : FStringView();
+	}
+	virtual const TCHAR* GetDataAsCString(int32 Key) const override
+	{
+		return Track && Key == TrackNameKey ? *Track->GetName() : nullptr;
+	}
+	virtual const FString& GetDataAsString(int32 Key) const override
+	{
+		return Track && Key == TrackNameKey ? Track->GetName() : EmptyString;
+	}
+
+	void SetTrack(const FBaseTimingTrack* InTrack)
+	{
+		Track = InTrack;
+	}
+
+private:
+	const FBaseTimingTrack* Track = nullptr;
+};
+
+////////////////////////////////////////////////////////////////////////////////////////////////////
+
 void STimingView::EnumerateFilteredTracks(TSharedPtr<UE::Insights::FFilterConfigurator> InFilterConfigurator, TSharedPtr<const FBaseTimingTrack> PriorityTrack, EnumerateFilteredTracksCallback Callback)
 {
-	using EFilterField = UE::Insights::EFilterField;
-	UE::Insights::FFilterContext FilterContext;
-	FilterContext.AddFilterData(static_cast<int32>(EFilterField::TrackName), FString());
+	FBaseTimingTrackFilterContext FilterContext;
 
 	// Call the callback for the PriorityTrack first if it passes the filters.
 	// This is an optimization because in many cases, the next/prev event will be on the same track
@@ -6089,7 +6135,7 @@ void STimingView::EnumerateFilteredTracks(TSharedPtr<UE::Insights::FFilterConfig
 		SkipId = PriorityTrack->GetId();
 		if (PriorityTrack->IsVisible())
 		{
-			FilterContext.SetFilterData(static_cast<int32>(EFilterField::TrackName), PriorityTrack->GetName());
+			FilterContext.SetTrack(PriorityTrack.Get());
 			if (InFilterConfigurator->ApplyFilters(FilterContext))
 			{
 				Callback(PriorityTrack);
@@ -6101,7 +6147,7 @@ void STimingView::EnumerateFilteredTracks(TSharedPtr<UE::Insights::FFilterConfig
 	{
 		if (Entry.Value->GetId() != SkipId && Entry.Value->IsVisible())
 		{
-			FilterContext.SetFilterData(static_cast<int32>(EFilterField::TrackName), Entry.Value->GetName());
+			FilterContext.SetTrack(Entry.Value.Get());
 			if (InFilterConfigurator->ApplyFilters(FilterContext))
 			{
 				Callback(Entry.Value);
diff --git a/Engine/Source/Developer/TraceInsightsCore/Private/InsightsCore/Filter/ViewModels/FilterConfigurator.cpp b/Engine/Source/Developer/TraceInsightsCore/Private/InsightsCore/Filter/ViewModels/FilterConfigurator.cpp
--- a/Engine/Source/Developer/TraceInsightsCore/Private/InsightsCore/Filter/ViewModels/FilterConfigurator.cpp
+++ b/Engine/Source/Developer/TraceInsightsCore/Private/InsightsCore/Filter/ViewModels/FilterConfigurator.cpp
@@ -68,7 +68,7 @@ bool FFilterConfigurator::operator==(const FFilterConfigurator& Other) const
 
 ////////////////////////////////////////////////////////////////////////////////////////////////////
 
-bool FFilterConfigurator::ApplyFilters(const FFilterContext& Context) const
+bool FFilterConfigurator::ApplyFilters(const FBaseFilterContext& Context) const
 {
 	return RootNode->ApplyFilters(Context);
 }
diff --git a/Engine/Source/Developer/TraceInsightsCore/Private/InsightsCore/Filter/ViewModels/FilterConfiguratorNode.cpp b/Engine/Source/Developer/TraceInsightsCore/Private/InsightsCore/Filter/ViewModels/FilterConfiguratorNode.cpp
--- a/Engine/Source/Developer/TraceInsightsCore/Private/InsightsCore/Filter/ViewModels/FilterConfiguratorNode.cpp
+++ b/Engine/Source/Developer/TraceInsightsCore/Private/InsightsCore/Filter/ViewModels/FilterConfiguratorNode.cpp
@@ -65,7 +65,7 @@ TSharedRef<FFilterConfiguratorNode> FFilterConfiguratorNode::DeepCopy(const FFil
 bool FFilterConfiguratorNode::operator==(const FFilterConfiguratorNode& Other) const
 {
 	bool bIsEqual = true;
-	
+
 	if (!IsGroup())
 	{
 		check(FilterState.IsValid());
@@ -189,44 +189,40 @@ void FFilterConfiguratorNode::GetUsedKeys(TSet<int32>& KeysUsed) const
 
 ////////////////////////////////////////////////////////////////////////////////////////////////////
 
-bool FFilterConfiguratorNode::ApplyFilters(const FFilterContext& Context) const
+bool FFilterConfiguratorNode::ApplyFilters(const FBaseFilterContext& Context) const
 {
-	bool Ret = true;
 	if (IsGroup())
 	{
 		switch (SelectedFilterGroupOperator->GetType())
 		{
 		case EFilterGroupOperator::And:
 		{
+			bool Ret = true;
 			auto& ChildrenArr = GetChildren();
 			for (int Index = 0; Index < ChildrenArr.Num() && Ret; ++Index)
 			{
 				Ret &= ((FFilterConfiguratorNode*)ChildrenArr[Index].Get())->ApplyFilters(Context);
 			}
-			break;
+			return Ret;
 		}
 		case EFilterGroupOperator::Or:
 		{
+			bool Ret = false;
 			auto& ChildrenArr = GetChildren();
-			if (ChildrenArr.Num() > 0)
-			{
-				Ret = false;
-			}
-
 			for (int Index = 0; Index < ChildrenArr.Num() && !Ret; ++Index)
 			{
 				Ret |= ((FFilterConfiguratorNode*)ChildrenArr[Index].Get())->ApplyFilters(Context);
 			}
-			break;
+			return Ret;
 		}
+		default:
+			return false;
 		}
 	}
 	else
 	{
-		Ret = FilterState->ApplyFilter(Context);
+		return FilterState->ApplyFilter(Context);
 	}
-
-	return Ret;
 }
 
 ////////////////////////////////////////////////////////////////////////////////////////////////////
diff --git a/Engine/Source/Developer/TraceInsightsCore/Private/InsightsCore/Filter/ViewModels/Filters.cpp b/Engine/Source/Developer/TraceInsightsCore/Private/InsightsCore/Filter/ViewModels/Filters.cpp
--- a/Engine/Source/Developer/TraceInsightsCore/Private/InsightsCore/Filter/ViewModels/Filters.cpp
+++ b/Engine/Source/Developer/TraceInsightsCore/Private/InsightsCore/Filter/ViewModels/Filters.cpp
@@ -15,6 +15,8 @@
 namespace UE::Insights
 {
 
+FString FBaseFilterContext::EmptyString;
+
 INSIGHTS_IMPLEMENT_RTTI(FFilterState)
 INSIGHTS_IMPLEMENT_RTTI(FFilter)
 INSIGHTS_IMPLEMENT_RTTI(FFilterWithSuggestions)
@@ -225,12 +227,12 @@ void FFilterState::SetFilterValue(FString InTextValue)
 
 ////////////////////////////////////////////////////////////////////////////////////////////////////
 
-bool FFilterState::ApplyFilter(const FFilterContext& Context) const
+bool FFilterState::ApplyFilter(const FBaseFilterContext& Context) const
 {
-	if (!Context.HasFilterData(Filter->GetKey()))
+	if (!Context.HasData(Filter->GetKey()))
 	{
 		// If data is not set for this filter return the value specified in the Context.
-		return Context.GetReturnValueForUnsetFilters();
+		return Context.GetDefaultReturnValue();
 	}
 
 	bool Ret = true;
@@ -240,8 +242,7 @@ bool FFilterState::ApplyFilter(const FFilterContext& Context) const
 	case EFilterDataType::Double:
 	{
 		FFilterOperator<double>* Operator = (FFilterOperator<double>*) SelectedOperator.Get();
-		double Value = 0.0;
-		Context.GetFilterData<double>(Filter->GetKey(), Value);
+		double Value = Context.GetDataAsDouble(Filter->GetKey());
 		Ret = Operator->Apply(Value, FilterValue.Get<double>());
 		break;
 	}
@@ -249,15 +250,14 @@ bool FFilterState::ApplyFilter(const FFilterContext& Context) const
 	case EFilterDataType::StringInt64Pair:
 	{
 		FFilterOperator<int64>* Operator = (FFilterOperator<int64>*) SelectedOperator.Get();
-		int64 Value = 0;
-		Context.GetFilterData<int64>(Filter->GetKey(), Value);
+		int64 Value = Context.GetDataAsInt64(Filter->GetKey());
 		Ret = Operator->Apply(Value, FilterValue.Get<int64>());
 		break;
 	}
 	case EFilterDataType::String:
 	{
 		FFilterOperator<FStringView>* Operator = (FFilterOperator<FStringView>*) SelectedOperator.Get();
-		const FString& Value = Context.GetFilterDataAsString(Filter->GetKey());
+		const FString& Value = Context.GetDataAsString(Filter->GetKey());
 		Ret = Operator->Apply(Value, FilterValue.Get<FString>());
 		break;
 	}
diff --git a/Engine/Source/Developer/TraceInsightsCore/Public/InsightsCore/Filter/ViewModels/FilterConfigurator.h b/Engine/Source/Developer/TraceInsightsCore/Public/InsightsCore/Filter/ViewModels/FilterConfigurator.h
--- a/Engine/Source/Developer/TraceInsightsCore/Public/InsightsCore/Filter/ViewModels/FilterConfigurator.h
+++ b/Engine/Source/Developer/TraceInsightsCore/Public/InsightsCore/Filter/ViewModels/FilterConfigurator.h
@@ -33,7 +33,7 @@ public:
 	/** Called to update the internal state of some filters */
 	void Update() { RootNode->Update(); }
 
-	UE_API virtual bool ApplyFilters(const FFilterContext& Context) const override;
+	UE_API virtual bool ApplyFilters(const FBaseFilterContext& Context) const override;
 
 	UE_API bool IsKeyUsed(int32 Key) const;
 
diff --git a/Engine/Source/Developer/TraceInsightsCore/Public/InsightsCore/Filter/ViewModels/FilterConfiguratorNode.h b/Engine/Source/Developer/TraceInsightsCore/Public/InsightsCore/Filter/ViewModels/FilterConfiguratorNode.h
--- a/Engine/Source/Developer/TraceInsightsCore/Public/InsightsCore/Filter/ViewModels/FilterConfiguratorNode.h
+++ b/Engine/Source/Developer/TraceInsightsCore/Public/InsightsCore/Filter/ViewModels/FilterConfiguratorNode.h
@@ -62,7 +62,7 @@ public:
 	const FString& GetTextBoxValue() { return TextBoxValue; }
 	UE_API void SetTextBoxValue(const FString& InValue);
 
-	UE_API bool ApplyFilters(const class FFilterContext& Context) const;
+	UE_API bool ApplyFilters(const FBaseFilterContext& Context) const;
 
 	UE_API void GetUsedKeys(TSet<int32>& GetUsedKeys) const;
 
diff --git a/Engine/Source/Developer/TraceInsightsCore/Public/InsightsCore/Filter/ViewModels/Filters.h b/Engine/Source/Developer/TraceInsightsCore/Public/InsightsCore/Filter/ViewModels/Filters.h
--- a/Engine/Source/Developer/TraceInsightsCore/Public/InsightsCore/Filter/ViewModels/Filters.h
+++ b/Engine/Source/Developer/TraceInsightsCore/Public/InsightsCore/Filter/ViewModels/Filters.h
@@ -25,14 +25,31 @@ namespace UE::Insights
 
 ////////////////////////////////////////////////////////////////////////////////////////////////////
 
-enum class EFilterDataType : uint32
+enum class EFilterDataType : uint8
 {
+	Void = 0, // not set
+
+	//Bool,
+	//Uint32,
+	//Uint64,
+	//Int32,
 	Int64,
+	//Float,
 	Double,
+	//StringView,
+	//AnsiStringView,
+	//Utf8StringView,
+	//CString,
+	//AnsiCString,
+	//Utf8CString,
 	String,
+	//AnsiString,
+	//Utf8String,
+	//ByteArray,
+	Custom, // For complex filters that are implemented as separate classes.
 	StringInt64Pair, // Displayed as a string but translates to a Int64 key.
 
-	Custom, // For complex filters that are implemented as separate classes.
+	Count
 };
 
 ////////////////////////////////////////////////////////////////////////////////////////////////////
@@ -146,12 +163,100 @@ public:
 
 ////////////////////////////////////////////////////////////////////////////////////////////////////
 
-class FFilterContext
+class FBaseFilterContext
 {
 public:
-	typedef TVariant<double, int64, FString> ContextData;
+	FBaseFilterContext() = default;
+	virtual ~FBaseFilterContext() = default;
+
+	virtual bool GetDefaultReturnValue() const
+	{
+		return false;
+	}
+
+	virtual bool HasData(int32 Key) const = 0;
+	virtual EFilterDataType GetDataType(int32 Key) const = 0;
+
+	//virtual bool GetDataAsBool(int32 Key) const                                 { return false; }
+	//virtual uint32 GetDataAsUint32(int32 Key) const                             { return 0; }
+	//virtual uint64 GetDataAsUint64(int32 Key) const                             { return 0; }
+	//virtual int32 GetDataAsInt32(int32 Key) const                               { return 0; }
+	virtual int64 GetDataAsInt64(int32 Key) const                               { return 0; }
+	//virtual float GetDataAsFloat(int32 Key) const                               { return 0.0f; }
+	virtual double GetDataAsDouble(int32 Key) const                             { return 0.0; }
+	virtual FStringView GetDataAsStringView(int32 Key) const                    { return GetDataAsString(Key); }
+	//virtual FAnsiStringView GetDataAsAnsiStringView(int32 Key) const            { return FAnsiStringView(); }
+	//virtual FUtf8StringView GetDataAsUtf8StringView(int32 Key) const            { return FUtf8StringView();}
+	virtual const TCHAR* GetDataAsCString(int32 Key) const                      { return *GetDataAsString(Key); }
+	//virtual const ANSICHAR* GetDataAsAnsiCString(int32 Key) const               { return nullptr; }
+	//virtual const UTF8CHAR* GetDataAsUtf8CString(int32 Key) const               { return nullptr; }
+	virtual const FString& GetDataAsString(int32 Key) const                     { return EmptyString; }
+	//virtual const FAnsiString& GetDataAsAnsiString(int32 Key) const             { return EmptyAnsiString; }
+	//virtual const FUtf8String& GetDataAsUtf8String(int32 Key) const             { return EmptyUtf8String; }
+	//virtual TTuple<const uint8*, uint64> GetDataAsByteArray(int32 Key) const    { return TTuple<const uint8*, uint64>(nullptr, 0); }
+	//virtual const void* GetDataAsPointer(int32 Key) const                       { return nullptr; }
 
+protected:
+	UE_API static FString EmptyString;
+	//static FAnsiString EmptyAnsiString;
+	//static FUtf8String EmptyUtf8String;
+};
+
+////////////////////////////////////////////////////////////////////////////////////////////////////
+
+class FFilterContext : public FBaseFilterContext
+{
 public:
+	typedef TVariant<int64, double, FString> ContextData;
+
+public:
+	FFilterContext() = default;
+	virtual ~FFilterContext() = default;
+
+	//////////////////////////////////////////////////
+
+	UE_DEPRECATED(5.8, "Use GetDefaultReturnValue() instead")
+	bool GetReturnValueForUnsetFilters() const { return bDefaultReturnValue; }
+
+	UE_DEPRECATED(5.8, "Use custom class instead")
+	void SetReturnValueForUnsetFilters(bool InValue) { bDefaultReturnValue = InValue; }
+
+	virtual bool GetDefaultReturnValue() const override { return bDefaultReturnValue; }
+
+	UE_DEPRECATED(5.8, "Use HasData() instead")
+	bool HasFilterData(int32 Key) const
+	{
+		return DataMap.Contains(Key);
+	}
+
+	virtual bool HasData(int32 Key) const override
+	{
+		return DataMap.Contains(Key);
+	}
+
+	virtual EFilterDataType GetDataType(int32 Key) const override
+	{
+		const ContextData* Data = DataMap.Find(Key);
+		if (Data)
+		{
+			if (Data->IsType<int64>())
+			{
+				return EFilterDataType::Int64;
+			}
+			if (Data->IsType<double>())
+			{
+				return EFilterDataType::Double;
+			}
+			if (Data->IsType<FString>())
+			{
+				return EFilterDataType::String;
+			}
+		}
+		return EFilterDataType::Void;
+	}
+
+	//////////////////////////////////////////////////
+
 	template<typename T>
 	void AddFilterData(int32 Key, const T& InData)
 	{
@@ -170,32 +275,44 @@ public:
 	void GetFilterData(int32 Key, T& OutData) const
 	{
 		const ContextData* Data = DataMap.Find(Key);
-		check(Data);
-
-		check(Data->IsType<T>());
+		check(Data && Data->IsType<T>());
 		OutData = Data->Get<T>();
 	}
 
+	UE_DEPRECATED(5.8, "Use GetDataAsString instead")
 	const FString& GetFilterDataAsString(int32 Key) const
 	{
 		const ContextData* Data = DataMap.Find(Key);
-		check(Data);
-
-		check(Data->IsType<FString>());
+		check(Data && Data->IsType<FString>());
 		return Data->Get<FString>();
 	}
 
-	bool HasFilterData(int32 Key) const
+	//////////////////////////////////////////////////
+
+	virtual int64 GetDataAsInt64(int32 Key) const override
 	{
-		return DataMap.Contains(Key);
+		const ContextData* Data = DataMap.Find(Key);
+		check(Data && Data->IsType<int64>())
+		return Data->Get<int64>();
+	}
+	virtual double GetDataAsDouble(int32 Key) const override
+	{
+		const ContextData* Data = DataMap.Find(Key);
+		check(Data && Data->IsType<double>());
+		return Data->Get<double>();
+	}
+	virtual const FString& GetDataAsString(int32 Key) const override
+	{
+		const ContextData* Data = DataMap.Find(Key);
+		check(Data && Data->IsType<FString>());
+		return Data->Get<FString>();
 	}
 
-	bool GetReturnValueForUnsetFilters() const { return bReturnValueForUnsetFilters; }
-	void SetReturnValueForUnsetFilters(bool InValue) { bReturnValueForUnsetFilters = InValue; }
+	//////////////////////////////////////////////////
 
 private:
 	TMap<int32, ContextData> DataMap;
-	bool bReturnValueForUnsetFilters = true;
+	bool bDefaultReturnValue = true;
 };
 
 ////////////////////////////////////////////////////////////////////////////////////////////////////
@@ -222,8 +339,9 @@ public:
 	void SetSelectedOperator(TSharedPtr<IFilterOperator> InOperator) { SelectedOperator = InOperator; }
 	TSharedPtr<const IFilterOperator> GetSelectedOperator() const { return SelectedOperator; }
 
-	virtual void Update() {};
-	UE_API virtual bool ApplyFilter(const FFilterContext& Context) const;
+	virtual void Update() {}
+
+	UE_API virtual bool ApplyFilter(const FBaseFilterContext& Context) const;
 
 	UE_API virtual void SetFilterValue(FString InTextValue);
 
diff --git a/Engine/Source/Developer/TraceInsightsCore/Public/InsightsCore/Filter/ViewModels/IFilterExecutor.h b/Engine/Source/Developer/TraceInsightsCore/Public/InsightsCore/Filter/ViewModels/IFilterExecutor.h
--- a/Engine/Source/Developer/TraceInsightsCore/Public/InsightsCore/Filter/ViewModels/IFilterExecutor.h
+++ b/Engine/Source/Developer/TraceInsightsCore/Public/InsightsCore/Filter/ViewModels/IFilterExecutor.h
@@ -7,12 +7,12 @@
 namespace UE::Insights
 {
 
-class FFilterContext;
+class FBaseFilterContext;
 
 class IFilterExecutor
 {
 public:
-	virtual bool ApplyFilters(const class FFilterContext& Context) const = 0;
+	virtual bool ApplyFilters(const FBaseFilterContext& Context) const = 0;
 };
 
 } // namespace UE::Insights
