[Insights] Timers: Fixed selection of a timer when the parent group was collapsed (now it will expand automatically).

#Synced-CL 48448406

diff --git a/Engine/Source/Developer/TraceInsights/Private/Insights/TimingProfiler/Widgets/STimersView.cpp b/Engine/Source/Developer/TraceInsights/Private/Insights/TimingProfiler/Widgets/STimersView.cpp
--- a/Engine/Source/Developer/TraceInsights/Private/Insights/TimingProfiler/Widgets/STimersView.cpp
+++ b/Engine/Source/Developer/TraceInsights/Private/Insights/TimingProfiler/Widgets/STimersView.cpp
@@ -2988,6 +2988,13 @@ void STimersView::SelectTimerNode(uint32 TimerId)
 	FTimerNodePtr NodePtr = GetTimerNode(TimerId);
 	if (NodePtr)
 	{
+		FBaseTreeNodePtr ParentNodePtr = NodePtr->GetParent();
+		while (ParentNodePtr.IsValid())
+		{
+			check(ParentNodePtr->Is<FTimerNode>());
+			TreeView->SetItemExpansion(StaticCastSharedPtr<FTimerNode>(ParentNodePtr), true);
+			ParentNodePtr = ParentNodePtr->GetParent();
+		}
 		TreeView->SetSelection(NodePtr);
 		TreeView->RequestScrollIntoView(NodePtr);
 	}
