[Insights]
* Callers/Callees: Added "% Parent" column. Shows the percentage of total inclusive duration (the selected timer relative to the parent timer).
* Callers/Callees: Added "% Root" column. Shows the percentage of total inclusive duration (the selected timer relative to the root timer).
* Timers/Callers/Callees: Added "Excl/Incl" column. Shows the percentage of total exclusive duration from the total inclusive duration.
* TraceInsightsCore: Added FDoubleValueFormatterAsPercent formatter class.

#Synced-CL 48603294

diff --git a/Engine/Source/Developer/TraceInsights/Private/Insights/TimingProfiler/ViewModels/TimersViewColumnFactory.cpp b/Engine/Source/Developer/TraceInsights/Private/Insights/TimingProfiler/ViewModels/TimersViewColumnFactory.cpp
--- a/Engine/Source/Developer/TraceInsights/Private/Insights/TimingProfiler/ViewModels/TimersViewColumnFactory.cpp
+++ b/Engine/Source/Developer/TraceInsights/Private/Insights/TimingProfiler/ViewModels/TimersViewColumnFactory.cpp
@@ -26,10 +26,12 @@ const FName FTimersViewColumns::MetaGroupNameColumnID(TEXT("MetaGroupName"));
 const FName FTimersViewColumns::TypeColumnID(TEXT("Type"));
 const FName FTimersViewColumns::InstanceCountColumnID(TEXT("Count"));
 const FName FTimersViewColumns::ChildInstanceCountColumnID(TEXT("ChildCount"));
-const FName FTimersViewColumns::AverageInstanceCountColumnID(TEXT("AverageCount")); 
+const FName FTimersViewColumns::AverageInstanceCountColumnID(TEXT("AverageCount"));
 
 // Inclusive Time columns
 const FName FTimersViewColumns::TotalInclusiveTimeColumnID(TEXT("TotalInclTime"));
+const FName FTimersViewColumns::ParentPercentageInclusiveTimeColumnID(TEXT("ParentPercentageInclTime"));
+const FName FTimersViewColumns::RootPercentageInclusiveTimeColumnID(TEXT("RootPercentageInclTime"));
 const FName FTimersViewColumns::MaxInclusiveTimeColumnID(TEXT("MaxInclTime"));
 const FName FTimersViewColumns::UpperQuartileInclusiveTimeColumnID(TEXT("UpperQuartileInclTime"));
 const FName FTimersViewColumns::AverageInclusiveTimeColumnID(TEXT("AverageInclTime"));
@@ -39,6 +41,7 @@ const FName FTimersViewColumns::MinInclusiveTimeColumnID(TEXT("MinInclTime"));
 
 // Exclusive Time columns
 const FName FTimersViewColumns::TotalExclusiveTimeColumnID(TEXT("TotalExclTime"));
+const FName FTimersViewColumns::PercentageExclusiveInclusiveTimeColumnID(TEXT("PercentageExclInclTime"));
 const FName FTimersViewColumns::MaxExclusiveTimeColumnID(TEXT("MaxExclTime"));
 const FName FTimersViewColumns::UpperQuartileExclusiveTimeColumnID(TEXT("UpperQuartileExclTime"));
 const FName FTimersViewColumns::AverageExclusiveTimeColumnID(TEXT("AverageExclTime"));
@@ -66,6 +69,7 @@ void FTimersViewColumnFactory::CreateTimersViewColumns(TArray<TSharedRef<FTableC
 	Columns.Add(CreateMinInclusiveTimeColumn());
 
 	Columns.Add(CreateTotalExclusiveTimeColumn());
+	Columns.Add(CreatePercentageExclusiveInclusiveTimeColumn());
 	Columns.Add(CreateMaxExclusiveTimeColumn());
 	Columns.Add(CreateAverageExclusiveTimeColumn());
 	Columns.Add(CreateMedianExclusiveTimeColumn());
@@ -74,7 +78,7 @@ void FTimersViewColumnFactory::CreateTimersViewColumns(TArray<TSharedRef<FTableC
 
 ////////////////////////////////////////////////////////////////////////////////////////////////////
 
-void FTimersViewColumnFactory::CreateTimerTreeViewColumns(TArray<TSharedRef<FTableColumn>>& Columns)
+void FTimersViewColumnFactory::CreateCallersViewColumns(TArray<TSharedRef<FTableColumn>>& Columns)
 {
 	Columns.Reset();
 
@@ -84,8 +88,31 @@ void FTimersViewColumnFactory::CreateTimerTreeViewColumns(TArray<TSharedRef<FTab
 	Columns.Add(CreateInstanceCountColumn());
 	Columns.Add(CreateChildInstanceCountColumn());
 	Columns.Add(CreateTotalInclusiveTimeColumn());
+	Columns.Add(CreateParentPercentageInclusiveTimeColumn());
+	Columns.Add(CreateRootPercentageInclusiveTimeColumn());
 	Columns.Add(CreateAverageInclusiveTimeColumn());
 	Columns.Add(CreateTotalExclusiveTimeColumn());
+	Columns.Add(CreatePercentageExclusiveInclusiveTimeColumn());
+	Columns.Add(CreateAverageExclusiveTimeColumn());
+}
+
+////////////////////////////////////////////////////////////////////////////////////////////////////
+
+void FTimersViewColumnFactory::CreateCalleesViewColumns(TArray<TSharedRef<FTableColumn>>& Columns)
+{
+	Columns.Reset();
+
+	Columns.Add(CreateNameColumn());
+	Columns.Add(CreateIdColumn());
+	Columns.Add(CreateTypeColumn());
+	Columns.Add(CreateInstanceCountColumn());
+	Columns.Add(CreateChildInstanceCountColumn());
+	Columns.Add(CreateTotalInclusiveTimeColumn());
+	Columns.Add(CreateParentPercentageInclusiveTimeColumn());
+	Columns.Add(CreateRootPercentageInclusiveTimeColumn());
+	Columns.Add(CreateAverageInclusiveTimeColumn());
+	Columns.Add(CreateTotalExclusiveTimeColumn());
+	Columns.Add(CreatePercentageExclusiveInclusiveTimeColumn());
 	Columns.Add(CreateAverageExclusiveTimeColumn());
 }
 
@@ -378,7 +405,7 @@ TSharedRef<FTableColumn> FTimersViewColumnFactory::CreateChildInstanceCountColum
 }
 
 ////////////////////////////////////////////////////////////////////////////////////////////////////
-	
+
 TSharedRef<FTableColumn> FTimersViewColumnFactory::CreateAverageInstanceCountColumn()
 {
 	TSharedRef<FTimersTableColumn> ColumnRef = MakeShared<FTimersTableColumn>(FTimersViewColumns::AverageInstanceCountColumnID);
@@ -386,7 +413,7 @@ TSharedRef<FTableColumn> FTimersViewColumnFactory::CreateAverageInstanceCountCol
 
 	Column.SetShortName(LOCTEXT("AverageInstanceCount_ColumnName", "C.Avg"));
 	Column.SetTitleName(LOCTEXT("AverageInstanceCount_ColumnTitle", "Average Instance Count"));
-	
+
 	Column.SetDescription(ETraceFrameType::TraceFrameType_Count,
 		LOCTEXT("AverageInstanceCount_ColumnDesc",
 			"Total number of timing event instances of the selected timer in the selected range\n"
@@ -486,6 +513,136 @@ TSharedRef<FTableColumn> FTimersViewColumnFactory::CreateTotalInclusiveTimeColum
 
 ////////////////////////////////////////////////////////////////////////////////////////////////////
 
+TSharedRef<FTableColumn> FTimersViewColumnFactory::CreateParentPercentageInclusiveTimeColumn()
+{
+	TSharedRef<FTimersTableColumn> ColumnRef = MakeShared<FTimersTableColumn>(FTimersViewColumns::ParentPercentageInclusiveTimeColumnID);
+	FTimersTableColumn& Column = *ColumnRef;
+
+	Column.SetShortName(LOCTEXT("ParentPercentageInclusiveTime_ColumnName", "% Parent"));
+	Column.SetTitleName(LOCTEXT("ParentPercentageInclusiveTime_ColumnTitle", "Total Inclusive Time / Parent's Total Inclusive Time (%)"));
+
+	FText Description = LOCTEXT("ParentPercentageInclusiveTime_ColumnDesc", "Percentage of total inclusive duration (the selected timer relative to the parent timer)");
+	Column.SetDescription(ETraceFrameType::TraceFrameType_Count, Description);
+	Column.SetDescription(ETraceFrameType::TraceFrameType_Game, Description);
+	Column.SetDescription(ETraceFrameType::TraceFrameType_Rendering, Description);
+
+	Column.SetFlags(ETableColumnFlags::CanBeHidden |
+		ETableColumnFlags::ShouldBeVisible |
+		ETableColumnFlags::CanBeFiltered);
+
+	Column.SetHorizontalAlignment(HAlign_Right);
+	Column.SetInitialWidth(PercentageColumnInitialWidth);
+
+	Column.SetDataType(ETableCellDataType::Double);
+	Column.SetAggregation(ETableColumnAggregation::Sum);
+
+	class FParentPercentageInclusiveTimeValueGetter : public FTableCellValueGetter
+	{
+	public:
+		virtual const TOptional<FTableCellValue> GetValue(const FTableColumn& Column, const FBaseTreeNode& Node) const
+		{
+			ensure(Column.GetId() == FTimersViewColumns::ParentPercentageInclusiveTimeColumnID);
+			if (Node.Is<FTimerNode>())
+			{
+				const FTimerNode& TimerNode = Node.As<FTimerNode>();
+				FBaseTreeNodePtr ParentNode = TimerNode.GetParent();
+				if (ParentNode.IsValid() && ParentNode->Is<FTimerNode>())
+				{
+					const double ParentTime = ParentNode->As<FTimerNode>().GetAggregatedStats().TotalInclusiveTime;
+					if (ParentTime > 0.0)
+					{
+						const double Time = TimerNode.GetAggregatedStats().TotalInclusiveTime;
+						return TOptional<FTableCellValue>(FTableCellValue(Time / ParentTime));
+					}
+				}
+			}
+			return TOptional<FTableCellValue>();
+		}
+	};
+
+	TSharedRef<ITableCellValueGetter> Getter = MakeShared<FParentPercentageInclusiveTimeValueGetter>();
+	Column.SetValueGetter(Getter);
+
+	TSharedRef<ITableCellValueFormatter> Formatter = MakeShared<FDoubleValueFormatterAsPercent>();
+	Column.SetValueFormatter(Formatter);
+
+	//TSharedRef<ITableCellValueSorter> Sorter = MakeShared<FSorterByDoubleValue>(ColumnRef);
+	TSharedRef<ITableCellValueSorter> Sorter = MakeShared<FTimerNodeSortingByTotalInclusiveTime>(ColumnRef);
+	Column.SetValueSorter(Sorter);
+	Column.SetInitialSortMode(EColumnSortMode::Descending);
+
+	return ColumnRef;
+}
+
+////////////////////////////////////////////////////////////////////////////////////////////////////
+
+TSharedRef<FTableColumn> FTimersViewColumnFactory::CreateRootPercentageInclusiveTimeColumn()
+{
+	TSharedRef<FTimersTableColumn> ColumnRef = MakeShared<FTimersTableColumn>(FTimersViewColumns::RootPercentageInclusiveTimeColumnID);
+	FTimersTableColumn& Column = *ColumnRef;
+
+	Column.SetShortName(LOCTEXT("RootPercentageInclusiveTime_ColumnName", "% Root"));
+	Column.SetTitleName(LOCTEXT("RootPercentageInclusiveTime_ColumnTitle", "Total Inclusive Time / Root's Total Inclusive Time (%)"));
+
+	FText Description = LOCTEXT("RootPercentageInclusiveTime_ColumnDesc", "Percentage of total inclusive duration (the selected timer relative to the root timer)");
+	Column.SetDescription(ETraceFrameType::TraceFrameType_Count, Description);
+	Column.SetDescription(ETraceFrameType::TraceFrameType_Game, Description);
+	Column.SetDescription(ETraceFrameType::TraceFrameType_Rendering, Description);
+
+	Column.SetFlags(ETableColumnFlags::CanBeHidden |
+		ETableColumnFlags::ShouldBeVisible |
+		ETableColumnFlags::CanBeFiltered);
+
+	Column.SetHorizontalAlignment(HAlign_Right);
+	Column.SetInitialWidth(PercentageColumnInitialWidth);
+
+	Column.SetDataType(ETableCellDataType::Double);
+	Column.SetAggregation(ETableColumnAggregation::Sum);
+
+	class FRootPercentageInclusiveTimeValueGetter : public FTableCellValueGetter
+	{
+	public:
+		virtual const TOptional<FTableCellValue> GetValue(const FTableColumn& Column, const FBaseTreeNode& Node) const
+		{
+			ensure(Column.GetId() == FTimersViewColumns::RootPercentageInclusiveTimeColumnID);
+			if (Node.Is<FTimerNode>())
+			{
+				const FTimerNode& TimerNode = Node.As<FTimerNode>();
+				FBaseTreeNodePtr RootNode = TimerNode.GetParent();
+				while (RootNode.IsValid() && RootNode->GetParent().IsValid())
+				{
+					RootNode = RootNode->GetParent();
+				}
+				if (RootNode.IsValid() && RootNode->Is<FTimerNode>())
+				{
+					const double RootTime = RootNode->As<FTimerNode>().GetAggregatedStats().TotalInclusiveTime;
+					if (RootTime > 0.0)
+					{
+						const double Time = TimerNode.GetAggregatedStats().TotalInclusiveTime;
+						return TOptional<FTableCellValue>(FTableCellValue(Time / RootTime));
+					}
+				}
+			}
+			return TOptional<FTableCellValue>();
+		}
+	};
+
+	TSharedRef<ITableCellValueGetter> Getter = MakeShared<FRootPercentageInclusiveTimeValueGetter>();
+	Column.SetValueGetter(Getter);
+
+	TSharedRef<ITableCellValueFormatter> Formatter = MakeShared<FDoubleValueFormatterAsPercent>();
+	Column.SetValueFormatter(Formatter);
+
+	//TSharedRef<ITableCellValueSorter> Sorter = MakeShared<FSorterByDoubleValue>(ColumnRef);
+	TSharedRef<ITableCellValueSorter> Sorter = MakeShared<FTimerNodeSortingByTotalInclusiveTime>(ColumnRef);
+	Column.SetValueSorter(Sorter);
+	Column.SetInitialSortMode(EColumnSortMode::Descending);
+
+	return ColumnRef;
+}
+
+////////////////////////////////////////////////////////////////////////////////////////////////////
+
 TSharedRef<FTableColumn> FTimersViewColumnFactory::CreateMaxInclusiveTimeColumn()
 {
 	TSharedRef<FTimersTableColumn> ColumnRef = MakeShared<FTimersTableColumn>(FTimersViewColumns::MaxInclusiveTimeColumnID);
@@ -774,6 +931,65 @@ TSharedRef<FTableColumn> FTimersViewColumnFactory::CreateTotalExclusiveTimeColum
 
 ////////////////////////////////////////////////////////////////////////////////////////////////////
 
+TSharedRef<FTableColumn> FTimersViewColumnFactory::CreatePercentageExclusiveInclusiveTimeColumn()
+{
+	TSharedRef<FTimersTableColumn> ColumnRef = MakeShared<FTimersTableColumn>(FTimersViewColumns::PercentageExclusiveInclusiveTimeColumnID);
+	FTimersTableColumn& Column = *ColumnRef;
+
+	Column.SetShortName(LOCTEXT("PercentageExclusiveInclusiveTime_ColumnName", "Excl/Incl"));
+	Column.SetTitleName(LOCTEXT("PercentageExclusiveInclusiveTime_ColumnTitle", "Total Exclusive Time / Total Inclusive Time (%)"));
+
+	FText Description = LOCTEXT("PercentageExclusiveInclusiveTime_ColumnDesc", "Percentage of total exclusive duration from the total inclusive duration");
+	Column.SetDescription(ETraceFrameType::TraceFrameType_Count, Description);
+	Column.SetDescription(ETraceFrameType::TraceFrameType_Game, Description);
+	Column.SetDescription(ETraceFrameType::TraceFrameType_Rendering, Description);
+
+	Column.SetFlags(ETableColumnFlags::CanBeHidden |
+		//ETableColumnFlags::ShouldBeVisible |
+		ETableColumnFlags::CanBeFiltered);
+
+	Column.SetHorizontalAlignment(HAlign_Right);
+	Column.SetInitialWidth(PercentageColumnInitialWidth);
+
+	Column.SetDataType(ETableCellDataType::Double);
+	Column.SetAggregation(ETableColumnAggregation::Sum);
+
+	class FPercentageExclusiveInclusiveTimeValueGetter : public FTableCellValueGetter
+	{
+	public:
+		virtual const TOptional<FTableCellValue> GetValue(const FTableColumn& Column, const FBaseTreeNode& Node) const
+		{
+			ensure(Column.GetId() == FTimersViewColumns::PercentageExclusiveInclusiveTimeColumnID);
+			if (Node.Is<FTimerNode>())
+			{
+				const FTimerNode& TimerNode = Node.As<FTimerNode>();
+				const double InclTime = TimerNode.GetAggregatedStats().TotalInclusiveTime;
+				if (InclTime > 0.0)
+				{
+					const double ExclTime = TimerNode.GetAggregatedStats().TotalExclusiveTime;
+					return TOptional<FTableCellValue>(FTableCellValue(ExclTime / InclTime));
+				}
+			}
+			return TOptional<FTableCellValue>();
+		}
+	};
+
+	TSharedRef<ITableCellValueGetter> Getter = MakeShared<FPercentageExclusiveInclusiveTimeValueGetter>();
+	Column.SetValueGetter(Getter);
+
+	TSharedRef<ITableCellValueFormatter> Formatter = MakeShared<FDoubleValueFormatterAsPercent>();
+	Column.SetValueFormatter(Formatter);
+
+	//TSharedRef<ITableCellValueSorter> Sorter = MakeShared<FSorterByDoubleValue>(ColumnRef);
+	TSharedRef<ITableCellValueSorter> Sorter = MakeShared<FTimerNodeSortingByTotalExclusiveTime>(ColumnRef);
+	Column.SetValueSorter(Sorter);
+	Column.SetInitialSortMode(EColumnSortMode::Descending);
+
+	return ColumnRef;
+}
+
+////////////////////////////////////////////////////////////////////////////////////////////////////
+
 TSharedRef<FTableColumn> FTimersViewColumnFactory::CreateMaxExclusiveTimeColumn()
 {
 	TSharedRef<FTimersTableColumn> ColumnRef = MakeShared<FTimersTableColumn>(FTimersViewColumns::MaxExclusiveTimeColumnID);
diff --git a/Engine/Source/Developer/TraceInsights/Private/Insights/TimingProfiler/ViewModels/TimersViewColumnFactory.h b/Engine/Source/Developer/TraceInsights/Private/Insights/TimingProfiler/ViewModels/TimersViewColumnFactory.h
--- a/Engine/Source/Developer/TraceInsights/Private/Insights/TimingProfiler/ViewModels/TimersViewColumnFactory.h
+++ b/Engine/Source/Developer/TraceInsights/Private/Insights/TimingProfiler/ViewModels/TimersViewColumnFactory.h
@@ -29,6 +29,8 @@ struct FTimersViewColumns
 
 	// Inclusive Time columns
 	static const FName TotalInclusiveTimeColumnID;
+	static const FName ParentPercentageInclusiveTimeColumnID;
+	static const FName RootPercentageInclusiveTimeColumnID;
 	static const FName MaxInclusiveTimeColumnID;
 	static const FName UpperQuartileInclusiveTimeColumnID;
 	static const FName AverageInclusiveTimeColumnID;
@@ -38,6 +40,7 @@ struct FTimersViewColumns
 
 	// Exclusive Time columns
 	static const FName TotalExclusiveTimeColumnID;
+	static const FName PercentageExclusiveInclusiveTimeColumnID;
 	static const FName MaxExclusiveTimeColumnID;
 	static const FName UpperQuartileExclusiveTimeColumnID;
 	static const FName AverageExclusiveTimeColumnID;
@@ -99,7 +102,8 @@ struct FTimersViewColumnFactory
 {
 public:
 	static void CreateTimersViewColumns(TArray<TSharedRef<FTableColumn>>& Columns);
-	static void CreateTimerTreeViewColumns(TArray<TSharedRef<FTableColumn>>& Columns);
+	static void CreateCallersViewColumns(TArray<TSharedRef<FTableColumn>>& Columns);
+	static void CreateCalleesViewColumns(TArray<TSharedRef<FTableColumn>>& Columns);
 
 	static TSharedRef<FTableColumn> CreateNameColumn();
 	static TSharedRef<FTableColumn> CreateIdColumn();
@@ -110,12 +114,15 @@ public:
 	static TSharedRef<FTableColumn> CreateAverageInstanceCountColumn();
 
 	static TSharedRef<FTableColumn> CreateTotalInclusiveTimeColumn();
+	static TSharedRef<FTableColumn> CreateParentPercentageInclusiveTimeColumn();
+	static TSharedRef<FTableColumn> CreateRootPercentageInclusiveTimeColumn();
 	static TSharedRef<FTableColumn> CreateMaxInclusiveTimeColumn();
 	static TSharedRef<FTableColumn> CreateAverageInclusiveTimeColumn();
 	static TSharedRef<FTableColumn> CreateMedianInclusiveTimeColumn();
 	static TSharedRef<FTableColumn> CreateMinInclusiveTimeColumn();
 
 	static TSharedRef<FTableColumn> CreateTotalExclusiveTimeColumn();
+	static TSharedRef<FTableColumn> CreatePercentageExclusiveInclusiveTimeColumn();
 	static TSharedRef<FTableColumn> CreateMaxExclusiveTimeColumn();
 	static TSharedRef<FTableColumn> CreateAverageExclusiveTimeColumn();
 	static TSharedRef<FTableColumn> CreateMedianExclusiveTimeColumn();
@@ -124,6 +131,7 @@ public:
 private:
 	static constexpr float TotalTimeColumnInitialWidth = 60.0f;
 	static constexpr float TimeMsColumnInitialWidth = 50.0f;
+	static constexpr float PercentageColumnInitialWidth = 50.0f;
 };
 
 ////////////////////////////////////////////////////////////////////////////////////////////////////
diff --git a/Engine/Source/Developer/TraceInsights/Private/Insights/TimingProfiler/Widgets/STimerTreeView.cpp b/Engine/Source/Developer/TraceInsights/Private/Insights/TimingProfiler/Widgets/STimerTreeView.cpp
--- a/Engine/Source/Developer/TraceInsights/Private/Insights/TimingProfiler/Widgets/STimerTreeView.cpp
+++ b/Engine/Source/Developer/TraceInsights/Private/Insights/TimingProfiler/Widgets/STimerTreeView.cpp
@@ -153,9 +153,10 @@ void STimerTreeView::InitCommandList()
 ////////////////////////////////////////////////////////////////////////////////////////////////////
 
 BEGIN_SLATE_FUNCTION_BUILD_OPTIMIZATION
-void STimerTreeView::Construct(const FArguments& InArgs, const FText& InViewName)
+void STimerTreeView::Construct(const FArguments& InArgs, const FText& InViewName, ETreeType InTreeType)
 {
 	ViewName = InViewName;
+	TreeType = InTreeType;
 
 	TSharedRef<FTimerButterflyAggregator> TimerButterflyAggregator = FTimingProfilerManager::Get()->GetTimerButterflyAggregator();
 
@@ -821,7 +822,18 @@ void STimerTreeView::InitializeAndShowHeaderColumns()
 {
 	// Create columns.
 	TArray<TSharedRef<FTableColumn>> Columns;
-	FTimersViewColumnFactory::CreateTimerTreeViewColumns(Columns);
+	if (TreeType == ETreeType::Callers)
+	{
+		FTimersViewColumnFactory::CreateCallersViewColumns(Columns);
+	}
+	else if (TreeType == ETreeType::Callees)
+	{
+		FTimersViewColumnFactory::CreateCalleesViewColumns(Columns);
+	}
+	else
+	{
+		FTimersViewColumnFactory::CreateTimersViewColumns(Columns);
+	}
 	if (ensure(Columns.Num() > 0 && Columns[0]->IsHierarchy()))
 	{
 		Columns[0]->SetShortName(ViewName);
diff --git a/Engine/Source/Developer/TraceInsights/Private/Insights/TimingProfiler/Widgets/STimerTreeView.h b/Engine/Source/Developer/TraceInsights/Private/Insights/TimingProfiler/Widgets/STimerTreeView.h
--- a/Engine/Source/Developer/TraceInsights/Private/Insights/TimingProfiler/Widgets/STimerTreeView.h
+++ b/Engine/Source/Developer/TraceInsights/Private/Insights/TimingProfiler/Widgets/STimerTreeView.h
@@ -45,6 +45,14 @@ class STimingView;
  */
 class STimerTreeView : public SCompoundWidget
 {
+public:
+	enum class ETreeType
+	{
+		Timers,
+		Callers,
+		Callees,
+	};
+
 public:
 	/** Default constructor. */
 	STimerTreeView();
@@ -59,7 +67,7 @@ public:
 	 * Construct this widget
 	 * @param InArgs - The declaration data for this widget
 	 */
-	void Construct(const FArguments& InArgs, const FText& InViewName);
+	void Construct(const FArguments& InArgs, const FText& InViewName, ETreeType InTreeType);
 
 	void Reset();
 
@@ -255,6 +263,8 @@ private:
 	/** Table view model. */
 	TSharedPtr<FTable> Table;
 
+	ETreeType TreeType;
+
 	/** The view name (ex.: "Callers" or "Callees"). */
 	FText ViewName;
 
diff --git a/Engine/Source/Developer/TraceInsights/Private/Insights/TimingProfiler/Widgets/STimingProfilerWindow.cpp b/Engine/Source/Developer/TraceInsights/Private/Insights/TimingProfiler/Widgets/STimingProfilerWindow.cpp
--- a/Engine/Source/Developer/TraceInsights/Private/Insights/TimingProfiler/Widgets/STimingProfilerWindow.cpp
+++ b/Engine/Source/Developer/TraceInsights/Private/Insights/TimingProfiler/Widgets/STimingProfilerWindow.cpp
@@ -235,7 +235,7 @@ TSharedRef<SDockTab> STimingProfilerWindow::SpawnTab_Callers(const FSpawnTabArgs
 		.ShouldAutosize(false)
 		.TabRole(ETabRole::PanelTab)
 		[
-			SAssignNew(CallersTreeView, STimerTreeView, LOCTEXT("CallersTreeViewName", "Callers"))
+			SAssignNew(CallersTreeView, STimerTreeView, LOCTEXT("CallersTreeViewName", "Callers"), STimerTreeView::ETreeType::Callers)
 		];
 
 	DockTab->SetOnTabClosed(SDockTab::FOnTabClosedCallback::CreateRaw(this, &STimingProfilerWindow::OnCallersTabClosed));
@@ -266,7 +266,7 @@ TSharedRef<SDockTab> STimingProfilerWindow::SpawnTab_Callees(const FSpawnTabArgs
 		.ShouldAutosize(false)
 		.TabRole(ETabRole::PanelTab)
 		[
-			SAssignNew(CalleesTreeView, STimerTreeView, LOCTEXT("CalleesTreeViewName", "Callees"))
+			SAssignNew(CalleesTreeView, STimerTreeView, LOCTEXT("CalleesTreeViewName", "Callees"), STimerTreeView::ETreeType::Callees)
 		];
 
 	DockTab->SetOnTabClosed(SDockTab::FOnTabClosedCallback::CreateRaw(this, &STimingProfilerWindow::OnCalleesTabClosed));
diff --git a/Engine/Source/Developer/TraceInsightsCore/Private/InsightsCore/Table/ViewModels/TableCellValueFormatter.cpp b/Engine/Source/Developer/TraceInsightsCore/Private/InsightsCore/Table/ViewModels/TableCellValueFormatter.cpp
--- a/Engine/Source/Developer/TraceInsightsCore/Private/InsightsCore/Table/ViewModels/TableCellValueFormatter.cpp
+++ b/Engine/Source/Developer/TraceInsightsCore/Private/InsightsCore/Table/ViewModels/TableCellValueFormatter.cpp
@@ -335,6 +335,33 @@ FText FDoubleValueFormatterAsCustomNumber::FormatValue(const TOptional<FTableCel
 
 ////////////////////////////////////////////////////////////////////////////////////////////////////
 
+FText FDoubleValueFormatterAsPercent::FormatValue(const TOptional<FTableCellValue>& InValue) const
+{
+	if (InValue.IsSet())
+	{
+		const double Value = InValue.GetValue().Double;
+		if (std::isnan(Value))
+		{
+			return FText::FromString(TEXT("NaN"));
+		}
+		else if (Value == 0.0)
+		{
+			return FText::FromString(TEXT("0%"));
+		}
+		else if (Value == 1.0)
+		{
+			return FText::FromString(TEXT("100%"));
+		}
+		else
+		{
+			return FText::AsPercent(Value, &FormattingOptions);
+		}
+	}
+	return FText::GetEmpty();
+}
+
+////////////////////////////////////////////////////////////////////////////////////////////////////
+
 FText FDoubleValueFormatterAsTimeAuto::FormatValue(const TOptional<FTableCellValue>& InValue) const
 {
 	if (InValue.IsSet())
diff --git a/Engine/Source/Developer/TraceInsightsCore/Public/InsightsCore/Table/ViewModels/TableCellValueFormatter.h b/Engine/Source/Developer/TraceInsightsCore/Public/InsightsCore/Table/ViewModels/TableCellValueFormatter.h
--- a/Engine/Source/Developer/TraceInsightsCore/Public/InsightsCore/Table/ViewModels/TableCellValueFormatter.h
+++ b/Engine/Source/Developer/TraceInsightsCore/Public/InsightsCore/Table/ViewModels/TableCellValueFormatter.h
@@ -200,7 +200,7 @@ public:
 };
 
 ////////////////////////////////////////////////////////////////////////////////////////////////////
-	
+
 class FDoubleValueFormatterAsCustomNumber : public FTableCellValueFormatter
 {
 public:
@@ -218,7 +218,28 @@ public:
 private:
 	FNumberFormattingOptions FormattingOptions;
 };
-	
+
+////////////////////////////////////////////////////////////////////////////////////////////////////
+
+class FDoubleValueFormatterAsPercent : public FTableCellValueFormatter
+{
+public:
+	FDoubleValueFormatterAsPercent()
+	{
+		FormattingOptions.MinimumFractionalDigits = 2;
+		FormattingOptions.MaximumFractionalDigits = 2;
+	}
+	virtual ~FDoubleValueFormatterAsPercent() = default;
+
+	const FNumberFormattingOptions& GetFormattingOptions() const { return FormattingOptions; }
+	FNumberFormattingOptions& GetFormattingOptions() { return FormattingOptions; }
+
+	UE_API virtual FText FormatValue(const TOptional<FTableCellValue>& InValue) const override;
+
+private:
+	FNumberFormattingOptions FormattingOptions;
+};
+
 ////////////////////////////////////////////////////////////////////////////////////////////////////
 
 class FDoubleValueFormatterAsTimeAuto : public FTableCellValueFormatter
