[Insight Settings] Fixed Fence Tracks defaulting to true instead of false

#rnx
#Synced-CL 49728215

diff --git a/Engine/Programs/UnrealInsights/Config/DefaultUnrealInsightsSettings.ini b/Engine/Programs/UnrealInsights/Config/DefaultUnrealInsightsSettings.ini
--- a/Engine/Programs/UnrealInsights/Config/DefaultUnrealInsightsSettings.ini
+++ b/Engine/Programs/UnrealInsights/Config/DefaultUnrealInsightsSettings.ini
@@ -22,7 +22,7 @@ ShowGpuWorkTracks=True
 ShowGpuWorkOverlays=True
 ShowGpuWorkExtendedLines=True
 ShowGpuFencesRelations=True
-ShowGpuFencesTracks=True
+ShowGpuFencesTracks=False
 ShowGpuFencesExtendedLines=True
 
 ShowAllVerseTracks=False
diff --git a/Engine/Source/Developer/TraceInsights/Private/Insights/InsightsSettings.h b/Engine/Source/Developer/TraceInsights/Private/Insights/InsightsSettings.h
--- a/Engine/Source/Developer/TraceInsights/Private/Insights/InsightsSettings.h
+++ b/Engine/Source/Developer/TraceInsights/Private/Insights/InsightsSettings.h
@@ -151,7 +151,7 @@ private:
 		RegisterToggleOptions(Option_ShowGpuWorkOverlays, true);
 		RegisterToggleOptions(Option_ShowGpuWorkExtendedLines, true);
 		RegisterToggleOptions(Option_ShowGpuFencesRelations, true);
-		RegisterToggleOptions(Option_ShowGpuFencesTracks, true);
+		RegisterToggleOptions(Option_ShowGpuFencesTracks, false);
 		RegisterToggleOptions(Option_ShowGpuFencesExtendedLines, true);
 
 		RegisterToggleOptions(Option_ShowAllVerseTracks, true);
