Remove auto focus on the task and cook minor tabs when they are spawned.
#Synced-CL 49071325

diff --git a/Engine/Source/Developer/TraceInsights/Private/Insights/CookProfiler/CookProfilerManager.cpp b/Engine/Source/Developer/TraceInsights/Private/Insights/CookProfiler/CookProfilerManager.cpp
--- a/Engine/Source/Developer/TraceInsights/Private/Insights/CookProfiler/CookProfilerManager.cpp
+++ b/Engine/Source/Developer/TraceInsights/Private/Insights/CookProfiler/CookProfilerManager.cpp
@@ -145,7 +145,7 @@ bool FCookProfilerManager::Tick(float DeltaTime)
 				if (CookProvider && CookProvider->GetNumPackages() > 0 && TabManagerShared.IsValid())
 				{
 					bIsAvailable = true;
-					TabManagerShared->TryInvokeTab(FCookProfilerTabs::PackageTableTreeViewTabID);
+					TabManagerShared->TryInvokeTab(FCookProfilerTabs::PackageTableTreeViewTabID, true);
 				}
 
 				if (Session->IsAnalysisComplete())
diff --git a/Engine/Source/Developer/TraceInsights/Private/Insights/TaskGraphProfiler/TaskGraphProfilerManager.cpp b/Engine/Source/Developer/TraceInsights/Private/Insights/TaskGraphProfiler/TaskGraphProfilerManager.cpp
--- a/Engine/Source/Developer/TraceInsights/Private/Insights/TaskGraphProfiler/TaskGraphProfilerManager.cpp
+++ b/Engine/Source/Developer/TraceInsights/Private/Insights/TaskGraphProfiler/TaskGraphProfilerManager.cpp
@@ -168,7 +168,7 @@ bool FTaskGraphProfilerManager::Tick(float DeltaTime)
 					TaskTimingSharedState = MakeShared<FTaskTimingSharedState>(TimingView.Get());
 					IModularFeatures::Get().RegisterModularFeature(Timing::TimingViewExtenderFeatureName, TaskTimingSharedState.Get());
 				}
-				TabManagerShared->TryInvokeTab(FTaskGraphProfilerTabs::TaskTableTreeViewTabID);
+				TabManagerShared->TryInvokeTab(FTaskGraphProfilerTabs::TaskTableTreeViewTabID, true);
 			}
 
 			if (Session->IsAnalysisComplete())
