Automation for a new Memory Leak Exporter (AaBCf rule) in background. It generates a snapshot csv ready for analysis.
Example command usage: UnrealInsights.exe -OpenTraceFile="atracefile.utrace" -unattended -autoquit -noui -nullrhi -ExecOnAnalysisCompleteCmd="[at]=ExportCommands.rsp" -log -AbsLog="MemoryLeakSnapshot.log"
ExportCommands.rsp example: MemoryInsights.ExportLeaks -Rule=AaBCf -Output=MemoryLeakSnapshot.csv -TimeA=0 -BookmarkB="bookmark_name" -BookmarkC="bookmark_name" -MaxResults="0" (zero means all query results)

Although this change only includes the AaBCf rule, it supports future development of other memory exporter rules.

#tests Locally, running tests via command line

#Synced-CL 49554259

#Synced-CL 49554259

diff --git a/Engine/Source/Developer/TraceInsights/Private/Insights/MemoryProfiler/MemoryExportCommandParser.cpp b/Engine/Source/Developer/TraceInsights/Private/Insights/MemoryProfiler/MemoryExportCommandParser.cpp
new file mode 100644
--- /dev/null
+++ b/Engine/Source/Developer/TraceInsights/Private/Insights/MemoryProfiler/MemoryExportCommandParser.cpp
@@ -0,0 +1,345 @@
+// Copyright Epic Games, Inc. All Rights Reserved.
+
+#include "MemoryExportCommandParser.h"
+
+namespace UE::Insights::MemoryProfiler
+{
+DEFINE_LOG_CATEGORY(LogMemoryCommandParser);
+
+FMemoryExportCommandParser::FMemoryExportCommandParser() :
+	TimeA(-1.0), TimeB(-1.0), TimeC(-1.0), TimeD(-1.0), MaxResults(0)
+{
+
+}
+
+bool FMemoryExportCommandParser::ParseParameters(const TCHAR* Cmd, FOutputDevice& Ar)
+{
+	UE_LOG(LogMemoryCommandParser, Log, TEXT("Validating parameters necessary to execute chosen rule..."));
+	const bool UseEscape = true;
+	while (Cmd && Cmd[0] != TEXT('\0'))
+	{
+		FString Token;
+		if (FParse::Token(Cmd, Token, UseEscape))
+		{
+			UE_LOG(LogMemoryCommandParser, Log, TEXT("Reading Parameter %s"), *Token);
+
+			static constexpr TCHAR RuleToken[] = TEXT("-Rule=");
+			static constexpr TCHAR OutputToken[] = TEXT("-Output=");
+			static constexpr TCHAR BookmarkAToken[] = TEXT("-BookmarkA=");
+			static constexpr TCHAR BookmarkBToken[] = TEXT("-BookmarkB=");
+			static constexpr TCHAR BookmarkCToken[] = TEXT("-BookmarkC=");
+			static constexpr TCHAR BookmarkDToken[] = TEXT("-BookmarkD=");
+			static constexpr TCHAR TimeAToken[] = TEXT("-TimeA=");
+			static constexpr TCHAR TimeBToken[] = TEXT("-TimeB=");
+			static constexpr TCHAR TimeCToken[] = TEXT("-TimeC=");
+			static constexpr TCHAR TimeDToken[] = TEXT("-TimeD=");
+			static constexpr TCHAR MaxResultsToken[] = TEXT("-MaxResults=");
+			static constexpr TCHAR ColumnsToken[] = TEXT("-Columns=");
+
+			if (Token.StartsWith(RuleToken))
+			{
+				Token.RightChopInline(UE_ARRAY_COUNT(RuleToken) - 1);
+				Token.TrimQuotesInline();
+				if (Token.IsEmpty())
+				{
+					UE_LOG(LogMemoryCommandParser, Error, TEXT("Exporter stopped because no Rule was provided. Please, specify one. Ex: -Rule=AaBCf"));
+					return false;
+				}
+				const QueryRule* RuleEnumPtr = QueryRuleMap().Find(Token);
+				if (RuleEnumPtr == nullptr)
+				{
+					UE_LOG(LogMemoryCommandParser, Error, TEXT("Exporter stopped because Rule (%s) is not a valid one"), *Token);
+					TArray<FString> RuleNames;
+					QueryRuleMap().GetKeys(RuleNames);
+					UE_LOG(LogMemoryCommandParser, Error, TEXT("Try one of the following Rules only: [%s]"), *FString::Join(RuleNames, TEXT(" | ")));
+					return false;
+				}
+				Rule = Token;
+				RuleEnum = *RuleEnumPtr;
+			}
+			else if (Token.StartsWith(OutputToken))
+			{
+				Token.RightChopInline(UE_ARRAY_COUNT(OutputToken) - 1);
+				Token.TrimQuotesInline();
+				OutputPath = Token;
+			}
+			else if (Token.StartsWith(BookmarkAToken))
+			{
+				Token.RightChopInline(UE_ARRAY_COUNT(BookmarkAToken) - 1);
+				Token.TrimQuotesInline();
+				BookmarkA = Token;
+			}
+			else if (Token.StartsWith(BookmarkBToken))
+			{
+				Token.RightChopInline(UE_ARRAY_COUNT(BookmarkBToken) - 1);
+				Token.TrimQuotesInline();
+				BookmarkB = Token;
+			}
+			else if (Token.StartsWith(BookmarkCToken))
+			{
+				Token.RightChopInline(UE_ARRAY_COUNT(BookmarkCToken) - 1);
+				Token.TrimQuotesInline();
+				BookmarkC = Token;
+			}
+			else if (Token.StartsWith(BookmarkDToken))
+			{
+				Token.RightChopInline(UE_ARRAY_COUNT(BookmarkDToken) - 1);
+				Token.TrimQuotesInline();
+				BookmarkD = Token;
+			}
+			else if (Token.StartsWith(TimeAToken))
+			{
+				Token.RightChopInline(UE_ARRAY_COUNT(TimeAToken) - 1);
+				if (Token.IsEmpty())
+				{
+					continue;
+				}
+				if (!Token.IsNumeric())
+				{
+					UE_LOG(LogMemoryCommandParser, Error, TEXT("TimeA value %s is not a valid number"), *Token);
+					return false;
+				}
+				TimeA = FCString::Atod(*Token);
+				if (TimeA < 0.0)
+				{
+					UE_LOG(LogMemoryCommandParser, Error, TEXT("Cannot pass a negative value (%.2f) for TimeA"), TimeA);
+					return false;
+				}
+			}
+			else if (Token.StartsWith(TimeBToken))
+			{
+				Token.RightChopInline(UE_ARRAY_COUNT(TimeBToken) - 1);
+				if (Token.IsEmpty())
+				{
+					continue;
+				}
+				if (!Token.IsNumeric())
+				{
+					UE_LOG(LogMemoryCommandParser, Error, TEXT("TimeB value %s is not a valid number"), *Token);
+					return false;
+				}
+				TimeB = FCString::Atod(*Token);
+				if (TimeB < 0.0)
+				{
+					UE_LOG(LogMemoryCommandParser, Error, TEXT("Cannot pass a negative value (%.2f) for TimeB"), TimeB);
+					return false;
+				}
+			}
+			else if (Token.StartsWith(TimeCToken))
+			{
+				Token.RightChopInline(UE_ARRAY_COUNT(TimeCToken) - 1);
+				if (Token.IsEmpty())
+				{
+					continue;
+				}
+				if (!Token.IsNumeric())
+				{
+					UE_LOG(LogMemoryCommandParser, Error, TEXT("TimeC value %s is not a valid number"), *Token);
+					return false;
+				}
+				TimeC = FCString::Atod(*Token);
+				if (TimeC < 0.0)
+				{
+					UE_LOG(LogMemoryCommandParser, Error, TEXT("Cannot pass a negative value (%.2f) for TimeC"), TimeC);
+					return false;
+				}
+			}
+			else if (Token.StartsWith(TimeDToken))
+			{
+				Token.RightChopInline(UE_ARRAY_COUNT(TimeDToken) - 1);
+				if (Token.IsEmpty())
+				{
+					continue;
+				}
+				if (!Token.IsNumeric())
+				{
+					UE_LOG(LogMemoryCommandParser, Error, TEXT("TimeD value %s is not a valid number"), *Token);
+					return false;
+				}
+				TimeD = FCString::Atod(*Token);
+				if (TimeD < 0.0)
+				{
+					UE_LOG(LogMemoryCommandParser, Error, TEXT("Cannot pass a negative value (%.2f) for TimeD"), TimeD);
+					return false;
+				}
+			}
+			else if (Token.StartsWith(MaxResultsToken))
+			{
+				// If user pass the parameter without values or with garbage, Atoi handles it
+				// and MaxResults will default to 0
+				Token.RightChopInline(UE_ARRAY_COUNT(MaxResultsToken) - 1);
+				MaxResults = FCString::Atoi(*Token);
+				if (MaxResults < 0)
+				{
+					UE_LOG(LogMemoryCommandParser, Error, TEXT("Cannot pass a negative value (%d) for MaxResults"), MaxResults);
+					return false;
+				}
+			}
+			else if (Token.StartsWith(ColumnsToken))
+			{
+				Token.RightChopInline(UE_ARRAY_COUNT(ColumnsToken) - 1);
+				Columns = Token;
+			}
+			else
+			{
+				UE_LOG(LogMemoryCommandParser, Warning, TEXT("Skipped unknown parameter: %s"), *Token);
+			}
+		}
+		else
+		{
+			UE_LOG(LogMemoryCommandParser, Error, TEXT("Failed parsing command parameters at: %s"), Cmd);
+			return false;
+		}
+	}
+	return true;
+}
+
+bool FMemoryExportCommandParser::ValidateParameters(const FMemoryExporter& Exporter)
+{
+	// Validate required parameters
+	if (Rule.IsEmpty())
+	{
+		UE_LOG(LogMemoryCommandParser, Error, TEXT("Rule must be specified. Ex: -Rule=AaBCf"));
+		return false;
+	}
+
+	// TODO: Implement all exporter rules. 
+	// Only AaBCf is implemented
+	if (RuleEnum != QueryRule::AaBCf)
+	{
+		UE_LOG(LogMemoryCommandParser, Error, TEXT("Rule '%s' not implemented yet via command line. Please use the UI."), *Rule);
+		return false;
+	}
+
+	// All rules need OutputPath
+	if (OutputPath.IsEmpty())
+	{
+		UE_LOG(LogMemoryCommandParser, Error, TEXT("Missing required parameter: -Output=\"path/to/file.csv\""));
+		return false;
+	}
+
+	// Validate Time usages based off of rule groups
+	if (!ValidTimeParameters(RuleEnum, RulesUsingTimeA(), BookmarkA, TimeA, TEXT("A"), Exporter))
+	{
+		return false;
+	}
+	if (!ValidTimeParameters(RuleEnum, RulesUsingTimeB(), BookmarkB, TimeB, TEXT("B"), Exporter))
+	{
+		return false;
+	}
+	if (!ValidTimeParameters(RuleEnum, RulesUsingTimeC(), BookmarkC, TimeC, TEXT("C"), Exporter))
+	{
+		return false;
+	}
+	if (!ValidTimeParameters(RuleEnum, RulesUsingTimeD(), BookmarkD, TimeD, TEXT("D"), Exporter))
+	{
+		return false;
+	}
+	return true;
+}
+
+bool FMemoryExportCommandParser::ValidTimeParameters(QueryRule RuleEnumerator, 
+	const TSet<QueryRule>& RulesUsingTimeX, 
+	const FString& BookmarkX, 
+	double& TimeX,
+	const TCHAR* TimeName,
+	const FMemoryExporter& Exporter)
+{
+	if (RulesUsingTimeX.Contains(RuleEnumerator))
+	{
+		if (BookmarkX.IsEmpty() && TimeX < 0.0)
+		{
+			UE_LOG(LogMemoryCommandParser, Error, 
+				TEXT("Time%s must be specified using either -Bookmark%s=<a bookmark name> or -Time%s=<a real positive number in seconds>"),
+				TimeName, TimeName, TimeName);
+			return false;
+		}
+		if (!BookmarkX.IsEmpty() && TimeX >= 0.0)
+		{
+			UE_LOG(LogMemoryCommandParser, Error, TEXT("Cannot specify both -Bookmark%s= and -Time%s=. Choose one."),
+				TimeName, TimeName);
+			return false;
+		}
+		// Only when a bookmark has been provided, then lookup the Timestamp
+		if (!BookmarkX.IsEmpty() && !Exporter.FindBookmarkByName(BookmarkX, TimeX))
+		{
+			UE_LOG(LogMemoryCommandParser, Error, TEXT("Bookmark '%s' not found, please re-run trace with valid bookmark name"), *BookmarkX);
+			return false;
+		}
+	}
+
+	return true;
+}
+
+// Certain rules need specific TimeX (Where X=A,B,C or D).
+// Since only one rule is implemented so far (AaBCf), the following functions
+// will serve for devs to add the rules they find using Time A,B,C or D.
+// As more rules are introduced in AllocationsProvider.h, they should be specified in these groups.
+
+const TSet<FMemoryExportCommandParser::QueryRule> FMemoryExportCommandParser::RulesUsingTimeA()
+{
+	static const TSet<QueryRule> Rules =
+	{
+		QueryRule::aAf, QueryRule::afA, QueryRule::Aaf, QueryRule::aAfB,
+		QueryRule::AaBf, QueryRule::aAfaBf, QueryRule::AfB, QueryRule::AaB,
+		QueryRule::AafB, QueryRule::aABf, QueryRule::AaBCf, QueryRule::AaBfC,
+		QueryRule::aABfC, QueryRule::AaBCfD, QueryRule::AoB, QueryRule::AiB
+	}; // TimeA is used in all rules
+
+	return Rules;
+}
+
+const TSet<FMemoryExportCommandParser::QueryRule> FMemoryExportCommandParser::RulesUsingTimeB()
+{
+	static const TSet<QueryRule> Rules =
+	{
+		QueryRule::aAfB, QueryRule::AaBf, QueryRule::aAfaBf, QueryRule::AfB, QueryRule::AaB,
+		QueryRule::AafB, QueryRule::aABf, QueryRule::AaBCf, QueryRule::AaBfC, QueryRule::aABfC,
+		QueryRule::AaBCfD, QueryRule::AoB, QueryRule::AiB
+	};
+	return Rules;
+}
+
+const TSet<FMemoryExportCommandParser::QueryRule> FMemoryExportCommandParser::RulesUsingTimeC()
+{
+	static const TSet<QueryRule> Rules =
+	{
+		QueryRule::AaBCf, QueryRule::AaBfC, QueryRule::aABfC, QueryRule::AaBCfD
+	};
+	return Rules;
+}
+
+const TSet<FMemoryExportCommandParser::QueryRule> FMemoryExportCommandParser::RulesUsingTimeD()
+{
+	static const TSet<QueryRule> Rules =
+	{
+		QueryRule::AaBCfD
+	};
+	return Rules;
+}
+
+const TMap<FString, FMemoryExportCommandParser::QueryRule> FMemoryExportCommandParser::QueryRuleMap()
+{
+	static const TMap<FString, QueryRule> Map =
+	{
+		{ TEXT("aAf"),     QueryRule::aAf },
+		{ TEXT("afA"),     QueryRule::afA },
+		{ TEXT("Aaf"),     QueryRule::Aaf },
+		{ TEXT("aAfB"),    QueryRule::aAfB },
+		{ TEXT("AaBf"),    QueryRule::AaBf },
+		{ TEXT("aAfaBf"),  QueryRule::aAfaBf },
+		{ TEXT("AfB"),     QueryRule::AfB },
+		{ TEXT("AaB"),     QueryRule::AaB },
+		{ TEXT("AafB"),    QueryRule::AafB },
+		{ TEXT("aABf"),    QueryRule::aABf },
+		{ TEXT("AaBCf"),   QueryRule::AaBCf },
+		{ TEXT("AaBfC"),   QueryRule::AaBfC },
+		{ TEXT("aABfC"),   QueryRule::aABfC },
+		{ TEXT("AaBCfD"),  QueryRule::AaBCfD },
+		{ TEXT("AoB"),     QueryRule::AoB },
+		{ TEXT("AiB"),     QueryRule::AiB }
+	};
+	return Map;
+}
+
+} // namespace UE::Insights::MemoryProfiler
diff --git a/Engine/Source/Developer/TraceInsights/Private/Insights/MemoryProfiler/MemoryExportCommandParser.h b/Engine/Source/Developer/TraceInsights/Private/Insights/MemoryProfiler/MemoryExportCommandParser.h
new file mode 100644
--- /dev/null
+++ b/Engine/Source/Developer/TraceInsights/Private/Insights/MemoryProfiler/MemoryExportCommandParser.h
@@ -0,0 +1,73 @@
+// Copyright Epic Games, Inc. All Rights Reserved.
+
+#pragma once
+
+#include "CoreTypes.h"
+
+// TraceServices
+#include "TraceServices/Model/AllocationsProvider.h"
+#include "TraceServices/Model/Memory.h"
+
+// ViewModels
+#include "ViewModels/MemoryExporter.h"
+
+
+namespace UE::Insights::MemoryProfiler
+{
+DECLARE_LOG_CATEGORY_EXTERN(LogMemoryCommandParser, Log, All);
+
+class FMemoryExportCommandParser
+{
+typedef TraceServices::IAllocationsProvider::EQueryRule QueryRule;
+
+public:
+	FMemoryExportCommandParser();
+	~FMemoryExportCommandParser() {}
+
+	bool ParseParameters(const TCHAR* Cmd, FOutputDevice& Ar);
+	bool ValidateParameters(const FMemoryExporter& Exporter);
+
+	// Getters
+	const FString& GetRuleName() const { return Rule; }
+	const QueryRule GetRuleEnum() const { return RuleEnum; }
+	const FString& GetOutputPathName() const { return OutputPath;  }
+	const FString& GetBookmarkA() const { return BookmarkA; }
+	const FString& GetBookmarkB() const { return BookmarkB; }
+	const FString& GetBookmarkC() const { return BookmarkC; }
+	const FString& GetBookmarkD() const { return BookmarkD; }
+	double GetTimeA() const { return TimeA; }
+	double GetTimeB() const { return TimeB; }
+	double GetTimeC() const { return TimeC; }
+	double GetTimeD() const { return TimeD; }
+	int32 GetMaxResults() const { return MaxResults; }
+	const FString& GetColumns() const { return Columns; }
+
+
+private:
+	bool ValidTimeParameters(QueryRule RuleEnumerator,
+		const TSet<QueryRule>& RulesUsingTime,
+		const FString& Bookmark,
+		double& Time,
+		const TCHAR* TimeName,
+		const FMemoryExporter& Exporter);
+
+	static const TSet<QueryRule> RulesUsingTimeA();
+	static const TSet<QueryRule> RulesUsingTimeB();
+	static const TSet<QueryRule> RulesUsingTimeC();
+	static const TSet<QueryRule> RulesUsingTimeD();
+	static const TMap<FString, QueryRule> QueryRuleMap();
+
+private:
+	FString Rule;
+	QueryRule RuleEnum;
+	FString OutputPath;
+	FString BookmarkA, BookmarkB, BookmarkC, BookmarkD;
+	double TimeA, TimeB, TimeC, TimeD;
+	int32 MaxResults;
+	FString Columns;
+
+		
+
+};
+
+} // namespace UE::Insights::MemoryProfiler
\ No newline at end of file
diff --git a/Engine/Source/Developer/TraceInsights/Private/Insights/MemoryProfiler/MemoryProfilerManager.cpp b/Engine/Source/Developer/TraceInsights/Private/Insights/MemoryProfiler/MemoryProfilerManager.cpp
--- a/Engine/Source/Developer/TraceInsights/Private/Insights/MemoryProfiler/MemoryProfilerManager.cpp
+++ b/Engine/Source/Developer/TraceInsights/Private/Insights/MemoryProfiler/MemoryProfilerManager.cpp
@@ -25,6 +25,10 @@
 #include "Insights/MemoryProfiler/Widgets/SMemoryProfilerWindow.h"
 #include "Insights/Widgets/STimingView.h"
 
+// Commad Line Exporters
+#include "MemoryExportCommandParser.h"
+#include "ViewModels/MemoryExporter.h"
+
 #define LOCTEXT_NAMESPACE "UE::Insights::MemoryProfiler"
 
 // For Memory Insights in editor preview.
@@ -451,6 +455,139 @@ void FMemoryProfilerManager::OnWindowClosedEvent()
 
 ////////////////////////////////////////////////////////////////////////////////////////////////////
 
+typedef TraceServices::IAllocationsProvider::EQueryRule QueryRule;
+static const TSet<QueryRule> RulesUsingTimeA =
+{
+	QueryRule::aAf, QueryRule::afA, QueryRule::Aaf, QueryRule::aAfB,
+	QueryRule::AaBf, QueryRule::aAfaBf, QueryRule::AfB, QueryRule::AaB,
+	QueryRule::AafB, QueryRule::aABf, QueryRule::AaBCf, QueryRule::AaBfC,
+	QueryRule::aABfC, QueryRule::AaBCfD, QueryRule::AoB, QueryRule::AiB
+}; // TimeA is used in all rules
+static const TSet<QueryRule> RulesUsingTimeB =
+{
+	QueryRule::aAfB, QueryRule::AaBf, QueryRule::aAfaBf, QueryRule::AfB, QueryRule::AaB,
+	QueryRule::AafB, QueryRule::aABf, QueryRule::AaBCf, QueryRule::AaBfC, QueryRule::aABfC,
+	QueryRule::AaBCfD, QueryRule::AoB, QueryRule::AiB
+};
+static const TSet<QueryRule> RulesUsingTimeC =
+{
+	QueryRule::AaBCf, QueryRule::AaBfC, QueryRule::aABfC, QueryRule::AaBCfD
+};
+static const TSet<QueryRule> RulesUsingTimeD =
+{
+	QueryRule::AaBCfD
+};
+static const TMap<FString, QueryRule> QueryRuleMap =
+{
+	{ TEXT("aAf"),     QueryRule::aAf },
+	{ TEXT("afA"),     QueryRule::afA },
+	{ TEXT("Aaf"),     QueryRule::Aaf },
+	{ TEXT("aAfB"),    QueryRule::aAfB },
+	{ TEXT("AaBf"),    QueryRule::AaBf },
+	{ TEXT("aAfaBf"),  QueryRule::aAfaBf },
+	{ TEXT("AfB"),     QueryRule::AfB },
+	{ TEXT("AaB"),     QueryRule::AaB },
+	{ TEXT("AafB"),    QueryRule::AafB },
+	{ TEXT("aABf"),    QueryRule::aABf },
+	{ TEXT("AaBCf"),   QueryRule::AaBCf },
+	{ TEXT("AaBfC"),   QueryRule::AaBfC },
+	{ TEXT("aABfC"),   QueryRule::aABfC },
+	{ TEXT("AaBCfD"),  QueryRule::AaBCfD },
+	{ TEXT("AoB"),     QueryRule::AoB },
+	{ TEXT("AiB"),     QueryRule::AiB }
+};
+
+static bool ValidateTimeParameters(QueryRule RuleEnum,
+	const TSet<QueryRule>& RulesUsingTime,
+	const FString& Bookmark,
+	double& Time,
+	const TCHAR* TimeName,
+	const FMemoryExporter& Exporter)
+{
+	return false;
+}
+
+bool FMemoryProfilerManager::Exec(const TCHAR* Cmd, FOutputDevice& Ar)
+{
+	// Parameterized Memory Leak Export Command
+	if (FParse::Command(&Cmd, TEXT("MemoryInsights.ExportAllocs")))
+	{
+		TSharedPtr<FInsightsManager> InsightsManager = FInsightsManager::Get();
+		if (!InsightsManager.IsValid())
+		{
+			UE_LOG(LogMemoryProfiler, Error, TEXT("InsightsManager is invalid"));
+			return false;
+		}
+
+		if (!InsightsManager->GetSession().IsValid())
+		{
+			UE_LOG(LogMemoryProfiler, Error, TEXT("Session is invalid"));
+			return false;
+		}
+
+		FMemoryExporter Exporter(*InsightsManager->GetSession().Get());
+
+		FMemoryExportCommandParser ExportCommandParser;
+		if (!ExportCommandParser.ParseParameters(Cmd, Ar))
+		{
+			UE_LOG(LogMemoryProfiler, Error, TEXT("Exporter stopped while parsing cmd parameters"));
+			return false;
+		}
+
+		if (!ExportCommandParser.ValidateParameters(Exporter))
+		{
+			UE_LOG(LogMemoryProfiler, Error, TEXT("Exporter stopped while validating parameters"));
+			return false;
+		}
+
+		// Below logic corresponds to Rule=AaBCf
+		if(ExportCommandParser.GetRuleEnum() == TraceServices::IAllocationsProvider::EQueryRule::AaBCf)
+		{
+			UE_LOG(LogMemoryProfiler, Log, TEXT(" === Memory Leak Exporter. Rule: \"%s\" === "), *ExportCommandParser.GetRuleName());
+			FMemoryExporter::FExportMemoryLeakSnapshotParams SnapshotParams;
+			SnapshotParams.OutputPathName = ExportCommandParser.GetOutputPathName();
+			SnapshotParams.TimeA = ExportCommandParser.GetTimeA();
+			SnapshotParams.TimeB = ExportCommandParser.GetTimeB();
+			SnapshotParams.TimeC = ExportCommandParser.GetTimeC();
+			SnapshotParams.MaxResults = ExportCommandParser.GetMaxResults();
+			SnapshotParams.Columns = ExportCommandParser.GetColumns();
+
+			FString BookmarkA = ExportCommandParser.GetBookmarkA();
+			FString BookmarkB = ExportCommandParser.GetBookmarkB();
+			FString BookmarkC = ExportCommandParser.GetBookmarkC();
+
+			UE_LOG(LogMemoryProfiler, Log, TEXT("Parameters provided by user:"));
+			UE_LOG(LogMemoryProfiler, Log, TEXT("-> Output:     %s"), *SnapshotParams.OutputPathName);
+			UE_LOG(LogMemoryProfiler, Log, TEXT("-> TimeA:      %.2f%s"), SnapshotParams.TimeA, BookmarkA.IsEmpty() ? TEXT("") : *FString::Printf(TEXT(" (from bookmark '%s')"), *BookmarkA));
+			UE_LOG(LogMemoryProfiler, Log, TEXT("-> TimeB:      %.2f%s"), SnapshotParams.TimeB, BookmarkB.IsEmpty() ? TEXT("") : *FString::Printf(TEXT(" (from bookmark '%s')"), *BookmarkB));
+			UE_LOG(LogMemoryProfiler, Log, TEXT("-> TimeC:      %.2f%s"), SnapshotParams.TimeC, BookmarkC.IsEmpty() ? TEXT("") : *FString::Printf(TEXT(" (from bookmark '%s')"), *BookmarkC));
+			if (ExportCommandParser.GetTimeD() >= 0.0 || !ExportCommandParser.GetBookmarkD().IsEmpty())
+			{
+				UE_LOG(LogMemoryProfiler, Warning, TEXT("-> TimeD SKIPPED. Not required in this rule"));
+			}
+			UE_LOG(LogMemoryProfiler, Log, TEXT("-> MaxResults: %d (0 = unlimited)"), SnapshotParams.MaxResults);
+			UE_LOG(LogMemoryProfiler, Log, TEXT("Starting export..."));
+			int ResultCount = Exporter.ExportMemoryLeakSnapshotAsText(SnapshotParams);
+
+			if (ResultCount >= 0)
+			{
+				UE_LOG(LogMemoryProfiler, Log, TEXT("SUCCESS! Exported %d allocations"), ResultCount);
+				UE_LOG(LogMemoryProfiler, Log, TEXT("Output file: %s"), *SnapshotParams.OutputPathName);
+				return true;
+			}
+			else
+			{
+				UE_LOG(LogMemoryProfiler, Error, TEXT("Export failed with code %d"), ResultCount);
+				return false;
+			}
+		}
+		
+	}
+	return false;
+}
+
+////////////////////////////////////////////////////////////////////////////////////////////////////
+
 } // namespace UE::Insights::MemoryProfiler
 
 #undef LOCTEXT_NAMESPACE
diff --git a/Engine/Source/Developer/TraceInsights/Private/Insights/MemoryProfiler/MemoryProfilerManager.h b/Engine/Source/Developer/TraceInsights/Private/Insights/MemoryProfiler/MemoryProfilerManager.h
--- a/Engine/Source/Developer/TraceInsights/Private/Insights/MemoryProfiler/MemoryProfilerManager.h
+++ b/Engine/Source/Developer/TraceInsights/Private/Insights/MemoryProfiler/MemoryProfilerManager.h
@@ -55,6 +55,7 @@ public:
 	virtual void Shutdown() override;
 	virtual void RegisterMajorTabs(IUnrealInsightsModule& InsightsModule) override;
 	virtual void UnregisterMajorTabs() override;
+	virtual bool Exec(const TCHAR* Cmd, FOutputDevice& Ar) override;
 	virtual void OnWindowClosedEvent() override;
 
 	//////////////////////////////////////////////////
diff --git a/Engine/Source/Developer/TraceInsights/Private/Insights/MemoryProfiler/ViewModels/MemoryExporter.cpp b/Engine/Source/Developer/TraceInsights/Private/Insights/MemoryProfiler/ViewModels/MemoryExporter.cpp
new file mode 100644
--- /dev/null
+++ b/Engine/Source/Developer/TraceInsights/Private/Insights/MemoryProfiler/ViewModels/MemoryExporter.cpp
@@ -0,0 +1,719 @@
+// Copyright Epic Games, Inc. All Rights Reserved.
+
+#include "MemoryExporter.h"
+
+#include "Containers/Map.h"
+#include "Containers/StringConv.h"
+#include "HAL/PlatformFileManager.h"
+#include "HAL/PlatformProcess.h"
+#include "Misc/Paths.h"
+#include "Logging/LogMacros.h"
+
+// TraceServices
+#include "Common/ProviderLock.h"
+#include "TraceServices/Model/AllocationsProvider.h"
+#include "TraceServices/Model/Bookmarks.h"
+#include "TraceServices/Model/Callstack.h"
+#include "TraceServices/Model/Definitions.h"
+#include "TraceServices/Model/MetadataProvider.h"
+#include "TraceServices/Model/Strings.h"
+
+
+DEFINE_LOG_CATEGORY_STATIC(LogMemoryLeakSnapshot, Log, All);
+
+namespace UE::Insights::MemoryProfiler
+{
+	///////////////////////////////////////////////////////////////////////////////////////////////////
+	// FMemoryExporter
+	////////////////////////////////////////////////////////////////////////////////////////////////////
+
+	FMemoryExporter::FMemoryExporter(const TraceServices::IAnalysisSession& InSession)
+		: Session(InSession)
+	{
+	}
+
+	////////////////////////////////////////////////////////////////////////////////////////////////////
+
+	FMemoryExporter::~FMemoryExporter()
+	{
+	}
+
+	////////////////////////////////////////////////////////////////////////////////////////////////////
+
+	const FName FMemoryExporter::ExportMemoryLeakSnapshot_AllocCallstack("AllocCallstack");
+	const FName FMemoryExporter::ExportMemoryLeakSnapshot_FreeCallstack("FreeCallstack");
+	
+	////////////////////////////////////////////////////////////////////////////////////////////////////
+
+	void FMemoryExporter::ExportMemorySnapshot_InitColumns()
+	{
+		InitializeAvailableColumns();
+		if (SnapshotDefaultColumns.IsEmpty())
+		{
+			SnapshotDefaultColumns = AvailableColumns;
+		}
+	}
+
+	void FMemoryExporter::InitializeAvailableColumns()
+	{
+		if (!AvailableColumns.IsEmpty())
+		{
+			return;
+		}
+
+		AvailableColumns.Add(FMemAllocTableColumns::StartEventIndexColumnId);
+		AvailableColumns.Add(FMemAllocTableColumns::EndEventIndexColumnId);
+		AvailableColumns.Add(FMemAllocTableColumns::EventDistanceColumnId);
+		AvailableColumns.Add(FMemAllocTableColumns::StartTimeColumnId);
+		AvailableColumns.Add(FMemAllocTableColumns::EndTimeColumnId);
+		AvailableColumns.Add(FMemAllocTableColumns::DurationColumnId);
+		AvailableColumns.Add(FMemAllocTableColumns::AddressColumnId);
+		AvailableColumns.Add(FMemAllocTableColumns::MemoryPageColumnId);
+		AvailableColumns.Add(FMemAllocTableColumns::SizeColumnId);
+		AvailableColumns.Add(FMemAllocTableColumns::TagColumnId);
+		AvailableColumns.Add(FMemAllocTableColumns::AssetColumnId);
+		AvailableColumns.Add(FMemAllocTableColumns::ClassNameColumnId);
+		AvailableColumns.Add(FMemAllocTableColumns::PackageColumnId);
+		AvailableColumns.Add(FMemAllocTableColumns::AllocThreadColumnId);
+		AvailableColumns.Add(FMemAllocTableColumns::FreeThreadColumnId);
+		AvailableColumns.Add(FMemAllocTableColumns::AllocFunctionColumnId);
+		AvailableColumns.Add(FMemAllocTableColumns::FreeFunctionColumnId);
+		AvailableColumns.Add(FMemAllocTableColumns::AllocSourceFileColumnId);
+		AvailableColumns.Add(FMemAllocTableColumns::FreeSourceFileColumnId);
+		AvailableColumns.Add(ExportMemoryLeakSnapshot_AllocCallstack);
+		AvailableColumns.Add(ExportMemoryLeakSnapshot_FreeCallstack);
+	}
+
+	////////////////////////////////////////////////////////////////////////////////////////////////////
+	// Copying the way TimingExporter generate the columns. 
+	// We need this function now to support handling specific columns
+	void FMemoryExporter::MakeExportMemoryLeakSnapshotColumnList(const FString& InColumnsString, TArray<FName>& OutColumnList)
+	{
+		InitializeAvailableColumns();
+
+		TArray<FString> Columns;
+		InColumnsString.ParseIntoArray(Columns, TEXT(","), true);
+
+		// Provide some hints of malformed list of columns (the white space case)
+		if (InColumnsString.StartsWith(TEXT(",")) || InColumnsString.EndsWith(TEXT(",")))
+		{
+			UE_LOG(LogMemoryLeakSnapshot, Warning, TEXT("Column list has leading/trailing comma: '%s'"), *InColumnsString);
+			UE_LOG(LogMemoryLeakSnapshot, Warning, 
+				TEXT("Check you didn't separate the list of columns with white spaces or remove extra commas."));
+		}
+
+		for (const FString& ColumnWildcard : Columns)
+		{
+			FName ColumnName(*ColumnWildcard);
+			if (AvailableColumns.Contains(ColumnName))
+			{
+				OutColumnList.Add(ColumnName);
+			}
+			else
+			{
+				for (const FName& Column : AvailableColumns)
+				{
+					if (Column.GetPlainNameString().MatchesWildcard(ColumnWildcard))
+					{
+						OutColumnList.Add(Column);
+					}
+				}
+			}
+		}
+	}
+
+	////////////////////////////////////////////////////////////////////////////////////////////////////
+
+	int32 FMemoryExporter::ExportMemoryLeakSnapshotAsText(FExportMemoryLeakSnapshotParams& Params)
+	{
+		UE_LOG(LogMemoryLeakSnapshot, Log, TEXT("Memory Leak Exporter Started"));
+		UE_LOG(LogMemoryLeakSnapshot, Log, TEXT("Output CSV file will be saved at: %s"), *Params.OutputPathName);
+		UE_LOG(LogMemoryLeakSnapshot, Log, TEXT("Memory Leak Analysis Times: A=%.2f, B=%.2f, C=%.2f"), Params.TimeA, Params.TimeB, Params.TimeC);
+
+		UE_LOG(LogMemoryLeakSnapshot, Log, TEXT("Querying allocations..."));
+		TArray<const TraceServices::IAllocationsProvider::FAllocation*> Allocations;
+		if (!QueryAllocations(Params, Allocations))
+		{
+			UE_LOG(LogMemoryLeakSnapshot, Error, TEXT("Failed to query allocations"));
+			return -1;
+		}
+		UE_LOG(LogMemoryLeakSnapshot, Log, TEXT("Found %d allocations"), Allocations.Num());
+
+		// Determine what columns to export
+		TArray<FName> ExportColumns;
+		if (Params.Columns.IsEmpty())
+		{
+			ExportMemorySnapshot_InitColumns();
+			ExportColumns = SnapshotDefaultColumns;
+			UE_LOG(LogMemoryLeakSnapshot, Log, TEXT("Using all %d columns available by default"), ExportColumns.Num());
+		}
+		else
+		{
+			MakeExportMemoryLeakSnapshotColumnList(Params.Columns, ExportColumns);
+			UE_LOG(LogMemoryLeakSnapshot, Log, TEXT("Using the %d columns passed in -Columns: %s"),
+				ExportColumns.Num(), *Params.Columns);
+		}
+		
+		// Adding limits just in case user needs for investigation
+		if (Params.MaxResults > 0 && Allocations.Num() > Params.MaxResults)
+		{
+			Allocations.SetNum(Params.MaxResults);
+			UE_LOG(LogMemoryLeakSnapshot, Log, TEXT("Export limited by the user to %d allocations"), Params.MaxResults);
+		}
+
+		// Write the CSV export
+		if (!WriteSnapshotCSV(Params.OutputPathName, Allocations, ExportColumns))
+		{
+			UE_LOG(LogMemoryLeakSnapshot, Error, TEXT("Failed to write CSV"));
+			return -1;
+		}
+
+
+		return Allocations.Num();
+	}
+
+	////////////////////////////////////////////////////////////////////////////////////////////////////
+
+	bool FMemoryExporter::QueryAllocations(const FExportMemoryLeakSnapshotParams& Params, TArray<const TraceServices::IAllocationsProvider::FAllocation*>& OutAllocations)
+	{
+		// Initialize required reading scope for sessions and providers
+
+		TraceServices::FAnalysisSessionReadScope SessionReadScope(Session);
+		const TraceServices::IAllocationsProvider* AllocationsProvider = TraceServices::ReadAllocationsProvider(Session);
+		if (!AllocationsProvider)
+		{
+			UE_LOG(LogMemoryLeakSnapshot, Error, TEXT("Nullptr: AllocationsProvider not available"));
+			return false;
+		}
+		TraceServices::FProviderReadScopeLock ProviderReadLock(*AllocationsProvider);
+		if (!AllocationsProvider->IsInitialized())
+		{
+			UE_LOG(LogMemoryLeakSnapshot, Error, TEXT("AllocationsProvider exists but not initialized"));
+			return false;
+		}
+
+		// Setup the query for memory leaks (AaBCf: It's one of the options from the Investigation tab in the UI)
+		TraceServices::IAllocationsProvider::FQueryParams QueryParams;
+		QueryParams.Rule = TraceServices::IAllocationsProvider::EQueryRule::AaBCf;
+		QueryParams.TimeA = Params.TimeA;
+		QueryParams.TimeB = Params.TimeB;
+		QueryParams.TimeC = Params.TimeC;
+		QueryParams.TimeD = 0.0;
+
+		// Start the query and poll results. The API is asynchronous and we have to collect results in batches from it's queue.
+		TraceServices::IAllocationsProvider::FQueryHandle QueryHandle = AllocationsProvider->StartQuery(QueryParams);
+		
+		TraceServices::IAllocationsProvider::FQueryStatus Status;
+		int BatchCount = 0;
+		do
+		{
+			Status = AllocationsProvider->PollQuery(QueryHandle);
+			if (Status.Status == TraceServices::IAllocationsProvider::EQueryStatus::Available)
+			{
+				TraceServices::IAllocationsProvider::FQueryResult Result = Status.NextResult();
+				if (Result)
+				{
+					const int NumAllocations = Result->Num();
+					UE_LOG(LogMemoryLeakSnapshot, Log, TEXT("Batch %d: %d allocations"), BatchCount, NumAllocations);
+					for (int i = 0; i < NumAllocations; ++i)
+					{
+						OutAllocations.Add(Result->Get(i));
+					}
+				}
+			}
+
+			++BatchCount;
+
+			// Handle time out with an arbitrarily big number for BatchCount
+			if (BatchCount > 10000)
+			{
+				UE_LOG(LogMemoryLeakSnapshot, Error, TEXT("Query timeout after %d batches iterated"), BatchCount);
+				AllocationsProvider->CancelQuery(QueryHandle);
+				return false;
+			}
+
+			// To save some cpu
+			FPlatformProcess::Sleep(0.001f);
+		} while (Status.Status != TraceServices::IAllocationsProvider::EQueryStatus::Done);
+
+
+		return true;
+	}
+
+	////////////////////////////////////////////////////////////////////////////////////////////////////
+	// This function appends in buffer in real time, because query may retrieve millions of allocations
+	// TODO: Evaluate using a Function table to avoid so many if/else if per row
+	void FMemoryExporter::WriteColumnValue(FUtf16Writer& Writer, const FName& ColumnId, 
+		const TraceServices::IAllocationsProvider::FAllocation* Alloc, 
+		const TraceServices::IAllocationsProvider* AllocationsProvider, 
+		const TraceServices::FCallstack* AllocCallstack, 
+		const TraceServices::FCallstack* FreeCallstack, 
+		const FString& AssetName, const FString& ClassName, const FString& PackageName, uint64 PageSize)
+	{
+		const uint32 MaxEventIndex = TNumericLimits<uint32>::Max();
+		if (ColumnId == FMemAllocTableColumns::StartEventIndexColumnId)
+		{
+			Writer.Append(LexToString(Alloc->GetStartEventIndex()));
+		}
+		else if (ColumnId == FMemAllocTableColumns::EndEventIndexColumnId)
+		{
+			Writer.Append(LexToString(Alloc->GetEndEventIndex()));
+		}
+		else if (ColumnId == FMemAllocTableColumns::EventDistanceColumnId)
+		{
+			const uint32 StartEventIndex = Alloc->GetStartEventIndex();
+			const uint32 EndEventIndex = Alloc->GetEndEventIndex();
+			const uint32 EventDistance = (EndEventIndex == MaxEventIndex) ? MaxEventIndex :
+				(EndEventIndex - StartEventIndex);
+			Writer.Append(LexToString(EventDistance));
+		}
+		else if (ColumnId == FMemAllocTableColumns::StartTimeColumnId)
+		{
+			Writer.Append(LexToString(Alloc->GetStartTime()));
+		}
+		else if (ColumnId == FMemAllocTableColumns::EndTimeColumnId)
+		{
+			const uint32 EndEventIndex = Alloc->GetEndEventIndex();
+			if (EndEventIndex == MaxEventIndex)
+			{
+				Writer.AppendString(TEXT("inf"));
+			}
+			else
+			{
+				Writer.Append(LexToString(Alloc->GetEndTime()));
+			}
+		}
+		else if (ColumnId == FMemAllocTableColumns::DurationColumnId)
+		{
+			const uint32 EndEventIndex = Alloc->GetEndEventIndex();
+			if (EndEventIndex == MaxEventIndex)
+			{
+				Writer.AppendString(TEXT("inf"));
+			}
+			else
+			{
+				Writer.Append(LexToString(Alloc->GetEndTime() - Alloc->GetStartTime()));
+			}
+		}
+		else if (ColumnId == FMemAllocTableColumns::AddressColumnId)
+		{
+			Writer.Append(LexToString(Alloc->GetAddress()));
+		}
+		else if (ColumnId == FMemAllocTableColumns::MemoryPageColumnId)
+		{
+			Writer.Append(LexToString(Alloc->GetAddress() / PageSize));
+		}
+		else if (ColumnId == FMemAllocTableColumns::SizeColumnId)
+		{
+			Writer.Append(LexToString(Alloc->GetSize()));
+		}
+		else if (ColumnId == FMemAllocTableColumns::TagColumnId)
+		{
+			const TCHAR* LLMTagPath = AllocationsProvider->GetTagFullPath(Alloc->GetTag());
+			Writer.AppendString(LLMTagPath ? LLMTagPath : TEXT(""));
+		}
+		else if (ColumnId == FMemAllocTableColumns::AssetColumnId)
+		{
+			Writer.AppendString(*AssetName);
+		}
+		else if (ColumnId == FMemAllocTableColumns::ClassNameColumnId)
+		{
+			Writer.AppendString(*ClassName);
+		}
+		else if (ColumnId == FMemAllocTableColumns::PackageColumnId)
+		{
+			Writer.AppendString(*PackageName);
+		}
+		else if (ColumnId == FMemAllocTableColumns::AllocThreadColumnId)
+		{
+			Writer.Append(LexToString(Alloc->GetAllocThreadId()));
+		}
+		else if (ColumnId == FMemAllocTableColumns::FreeThreadColumnId)
+		{
+			Writer.Append(LexToString(Alloc->GetFreeThreadId()));
+		}
+		else if (ColumnId == FMemAllocTableColumns::AllocFunctionColumnId)
+		{
+			Writer.AppendString(*GetTopFunction(AllocCallstack));
+		}
+		else if (ColumnId == FMemAllocTableColumns::FreeFunctionColumnId)
+		{
+			Writer.AppendString(*GetTopFunction(FreeCallstack));
+		}
+		else if (ColumnId == FMemAllocTableColumns::AllocSourceFileColumnId)
+		{
+			Writer.AppendString(*GetTopSourceFile(AllocCallstack));
+		}
+		else if (ColumnId == FMemAllocTableColumns::FreeSourceFileColumnId)
+		{
+			Writer.AppendString(*GetTopSourceFile(FreeCallstack));
+		}
+		else if (ColumnId == ExportMemoryLeakSnapshot_AllocCallstack)
+		{
+			Writer.AppendString(*FormatCallstackForCSV(AllocCallstack));
+		}
+		else if (ColumnId == ExportMemoryLeakSnapshot_FreeCallstack)
+		{
+			Writer.AppendString(*FormatCallstackForCSV(FreeCallstack));
+		}
+		else
+		{
+			Writer.AppendString(TEXT(""));
+		}
+	}
+
+	////////////////////////////////////////////////////////////////////////////////////////////////////
+
+
+	void FMemoryExporter::ExtractMetadata(const TraceServices::IAllocationsProvider::FAllocation* Allocation,
+		uint16 AssetMetadataType,
+		const TraceServices::IMetadataProvider* MetadataProvider,
+		const TraceServices::FMetadataSchema* MetadataSchema,
+		const TraceServices::IDefinitionProvider* DefinitionProvider,
+		FString& OutAssetName, FString& OutClassName, FString& OutPackageName)
+	{
+		const uint32 MetadataId = Allocation->GetMetadataId();
+		if (MetadataId != TraceServices::IMetadataProvider::InvalidMetadataId &&
+			MetadataProvider && DefinitionProvider && MetadataSchema)
+		{
+			TraceServices::FProviderReadScopeLock MetadataProviderReadLock(*MetadataProvider);
+			TraceServices::FProviderReadScopeLock DefinitionProviderReadLock(*DefinitionProvider);
+
+			MetadataProvider->EnumerateMetadata(Allocation->GetAllocThreadId(), MetadataId,
+				[AssetMetadataType, MetadataSchema, &OutAssetName, &OutClassName, &OutPackageName, DefinitionProvider]
+			(uint32 StackDepth, uint16 Type, const void* Data, uint32 Size) -> bool
+				{
+					if (Type == AssetMetadataType)
+					{
+						const auto Reader = MetadataSchema->Reader();
+
+						// Asset Name
+						const auto AssetNameRef = Reader.GetValueAs<UE::Trace::FEventRef32>((uint8*)Data, 0);
+						const auto AssetNameDef = DefinitionProvider->Get<TraceServices::FStringDefinition>(*AssetNameRef);
+						if (AssetNameDef)
+						{
+							OutAssetName = FString(AssetNameDef->Display);
+						}
+
+						// Class Name
+						const auto ClassNameRef = Reader.GetValueAs<UE::Trace::FEventRef32>((uint8*)Data, 1);
+						const auto ClassNameDef = DefinitionProvider->Get<TraceServices::FStringDefinition>(*ClassNameRef);
+						if (ClassNameDef)
+						{
+							OutClassName = FString(ClassNameDef->Display);
+						}
+
+						// Package Name
+						const auto PackageNameRef = Reader.GetValueAs<UE::Trace::FEventRef32>((uint8*)Data, 2);
+						const auto PackageNameDef = DefinitionProvider->Get<TraceServices::FStringDefinition>(*PackageNameRef);
+						if (PackageNameDef)
+						{
+							OutPackageName = FString(PackageNameDef->Display);
+						}
+
+						return false;
+					}
+					return true;
+				});
+		}
+	}
+
+	////////////////////////////////////////////////////////////////////////////////////////////////////
+
+	bool FMemoryExporter::WriteSnapshotCSV(const FString& FilePath, 
+		const TArray<const TraceServices::IAllocationsProvider::FAllocation*>& Allocations,
+		const TArray<FName>& ExportColumns)
+	{
+		// Deal with sessions and providers first
+		TraceServices::FAnalysisSessionReadScope SessionReadScope(Session);
+		const TraceServices::IAllocationsProvider* AllocationsProvider = TraceServices::ReadAllocationsProvider(Session);
+		const TraceServices::ICallstacksProvider* CallstacksProvider = TraceServices::ReadCallstacksProvider(Session);
+		const TraceServices::IMetadataProvider* MetadataProvider = TraceServices::ReadMetadataProvider(Session);
+		const TraceServices::IDefinitionProvider* DefinitionProvider = TraceServices::ReadDefinitionProvider(Session);
+		if (!AllocationsProvider)
+		{
+			UE_LOG(LogMemoryLeakSnapshot, Error, TEXT("AllocationsProvider not available"));
+			return false;
+		}
+		TraceServices::FProviderReadScopeLock AllocProviderLock(*AllocationsProvider);
+
+
+		// User requires data for Asset/Class/Package extraction. Setting up the Metadata schema here
+		uint16 AssetMetadataType = TraceServices::IMetadataProvider::InvalidMetadataType;
+		const TraceServices::FMetadataSchema* MetadataSchema = nullptr;
+		if (MetadataProvider)
+		{
+			TraceServices::FProviderReadScopeLock MetadataProviderLock(*MetadataProvider);
+			AssetMetadataType = MetadataProvider->GetRegisteredMetadataType(TEXT("Asset"));
+			if (AssetMetadataType != TraceServices::IMetadataProvider::InvalidMetadataType)
+			{
+				MetadataSchema = MetadataProvider->GetRegisteredMetadataSchema(AssetMetadataType);
+			}
+			else
+			{
+				// If Asset metadata type not found, disable metadata provider
+				MetadataProvider = nullptr;
+			}
+		}
+
+
+		// Building the CSV data
+		TUniquePtr<IFileHandle> FileHandle(FPlatformFileManager::Get().GetPlatformFile().OpenWrite(*FilePath));
+		if (!FileHandle)
+		{
+			UE_LOG(LogMemoryLeakSnapshot, Error, TEXT("Failed to open file for writing: %s"), *FilePath);
+			return false;
+		}
+		bool bIsCSV = FilePath.EndsWith(TEXT(".csv"));
+		FUtf16Writer Writer(FileHandle.Get(), bIsCSV);
+
+		// Write header - Generic support for N ExportColumns
+		for (int32 i = 0; i < ExportColumns.Num(); i++)
+		{
+			Writer.Append(ExportColumns[i].ToString());
+			if (i < ExportColumns.Num() - 1)
+			{
+				Writer.AppendSeparator();
+			}
+		}
+		Writer.AppendLineEnd();
+
+		// Write values
+		const uint32 MaxEventIndex = TNumericLimits<uint32>::Max();
+		const uint64 PageSize = 65536;
+		for (const TraceServices::IAllocationsProvider::FAllocation* Alloc : Allocations)
+		{
+			// Resolve callstacks
+			const TraceServices::FCallstack* AllocCallstack = nullptr;
+			const TraceServices::FCallstack* FreeCallstack = nullptr;
+			if (CallstacksProvider)
+			{
+				AllocCallstack = CallstacksProvider->GetCallstack(Alloc->GetAllocCallstackId());
+				FreeCallstack = CallstacksProvider->GetCallstack(Alloc->GetFreeCallstackId());
+			}
+
+			// Resolve Metadata
+			FString AssetName;
+			FString ClassName;
+			FString PackageName;
+			ExtractMetadata(Alloc, AssetMetadataType, MetadataProvider, MetadataSchema, DefinitionProvider,
+				AssetName, ClassName, PackageName);
+
+			// Writing Values - Generic
+			for (int32 i = 0; i < ExportColumns.Num(); ++i)
+			{
+				WriteColumnValue(Writer, ExportColumns[i], Alloc, AllocationsProvider, AllocCallstack, FreeCallstack,
+					AssetName, ClassName, PackageName, PageSize);
+				if (i < ExportColumns.Num() - 1)
+				{
+					Writer.AppendSeparator();
+				}
+			}
+			Writer.AppendLineEnd();
+		}
+
+		Writer.Flush();
+		UE_LOG(LogMemoryLeakSnapshot, Log, TEXT("Written %d rows to %s"), Allocations.Num(), *FilePath);
+		return true;
+	}
+
+	////////////////////////////////////////////////////////////////////////////////////////////////////
+
+	FString FMemoryExporter::GetTopFunction(const TraceServices::FCallstack* Callstack)
+	{
+		if (!Callstack || Callstack->IsEmpty())
+		{
+			return TEXT("No Callstack Recorded");
+		}
+
+		const TraceServices::FStackFrame* Frame = Callstack->Frame(0);
+		if (Frame && Frame->Symbol)
+		{
+			TStringBuilder<512> Builder;
+			if (Frame->Symbol->Module)
+			{
+				Builder.Append(Frame->Symbol->Module);
+				Builder.Append(TEXT("!"));
+			}
+			if (Frame->Symbol->Name)
+			{
+				Builder.Append(Frame->Symbol->Name);
+			}
+			else
+			{
+				Builder.Append(TEXT("0x"));
+				Builder.Append(LexToString(Frame->Addr));
+			}
+			return Builder.ToString();
+		}
+
+		return TEXT("No Callstack Recorded");
+	}
+
+	////////////////////////////////////////////////////////////////////////////////////////////////////
+
+	FString FMemoryExporter::GetTopSourceFile(const TraceServices::FCallstack* Callstack)
+	{
+		if (!Callstack || Callstack->IsEmpty())
+		{
+			return TEXT("No Callstack Recorded");
+		}
+
+		const TraceServices::FStackFrame* Frame = Callstack->Frame(0);
+		if (Frame && Frame->Symbol && Frame->Symbol->File)
+		{
+			return FString(Frame->Symbol->File);
+		}
+
+		return TEXT("No Callstack Recorded");
+	}
+
+	////////////////////////////////////////////////////////////////////////////////////////////////////
+
+	FString FMemoryExporter::FormatCallstackForCSV(const TraceServices::FCallstack* Callstack)
+	{
+		if (!Callstack || Callstack->IsEmpty())
+		{
+			return TEXT("No Callstack Recorded");
+		}
+
+		TStringBuilder<4096> Builder;
+		const uint32 NumFrames = Callstack->Num();
+		check(NumFrames < 256);
+
+		for (uint8 i = 0; i < NumFrames; ++i)
+		{
+			const TraceServices::FStackFrame* Frame = Callstack->Frame(i);
+			if (Frame && Frame->Symbol)
+			{
+				if (i > 0)
+				{
+					Builder.Append(TEXT("/"));
+				}
+
+				// Module!Function format
+				if (Frame->Symbol->Module)
+				{
+					Builder.Append(Frame->Symbol->Module);
+					Builder.Append(TEXT("!"));
+				}
+				if (Frame->Symbol->Name)
+				{
+					Builder.Append(Frame->Symbol->Name);
+				}
+				else
+				{
+					// Address if no symbol name
+					Builder.Append(TEXT("0x"));
+					Builder.Append(LexToString(Frame->Addr));
+				}
+			}
+		}
+
+		if (Builder.Len() == 0)
+		{
+			return TEXT("No Callstack Recorded");
+		}
+
+		return Builder.ToString();
+	}
+
+	////////////////////////////////////////////////////////////////////////////////////////////////////
+
+	bool FMemoryExporter::FindBookmarkByName(const FString& BookmarkName, double& OutTime) const
+	{
+		TraceServices::FAnalysisSessionReadScope SessionReadScope(Session);
+		const TraceServices::IBookmarkProvider& BookmarkProvider = TraceServices::ReadBookmarkProvider(Session);
+		const uint64 TotalBookmarks = BookmarkProvider.GetBookmarkCount();
+		if (TotalBookmarks == 0)
+		{
+			UE_LOG(LogMemoryLeakSnapshot, Warning, TEXT("No bookmarks found in trace"));
+			return false;
+		}
+		
+		bool BookmarkExists = false;
+		BookmarkProvider.EnumerateBookmarks(
+			0.0,
+			TNumericLimits<double>::Max(),
+			[&BookmarkName, &OutTime, &BookmarkExists](const TraceServices::FBookmark& Bookmark)
+			{
+				if (Bookmark.Text && FCString::Strcmp(Bookmark.Text, *BookmarkName) == 0)
+				{
+					OutTime = Bookmark.Time;
+					BookmarkExists = true;
+				}
+			});
+
+		if (BookmarkExists)
+		{
+			UE_LOG(LogMemoryLeakSnapshot, Log, TEXT("Found bookmark '%s' at time %.3f"), *BookmarkName, OutTime);
+		}
+		else
+		{
+			UE_LOG(LogMemoryLeakSnapshot, Warning, TEXT("Bookmark '%s' not found"), *BookmarkName);
+		}
+
+		return BookmarkExists;
+	}
+
+	////////////////////////////////////////////////////////////////////////////////////////////////////
+
+	void FMemoryExporter::FUtf16Writer::WriteBOM()
+	{
+		UTF16CHAR BOM = UNICODE_BOM;
+		FileHandle->Write((uint8*)&BOM, sizeof(UTF16CHAR));
+	}
+
+	////////////////////////////////////////////////////////////////////////////////////////////////////
+
+	void FMemoryExporter::FUtf16Writer::WriteStringBuilder(int32 CacheLen)
+	{
+		if (StringBuilder.Len() > CacheLen)
+		{
+			FTCHARToUTF16 UTF16String(StringBuilder.ToString(), StringBuilder.Len());
+			FileHandle->Write((const uint8*)UTF16String.Get(), UTF16String.Length() * sizeof(UTF16CHAR));
+			StringBuilder.Reset();
+		}
+	}
+
+	////////////////////////////////////////////////////////////////////////////////////////////////////
+
+	void FMemoryExporter::FUtf16Writer::AppendString(const TCHAR* Str)
+	{
+		if (Str == nullptr || Str[0] == TCHAR('\0'))
+		{
+			// nothing to append
+		}
+		else if (FCString::Strchr(Str, TCHAR(Separator)) != nullptr)
+		{
+			if (FCString::Strchr(Str, TEXT('\"')) != nullptr)
+			{
+				FString String = Str;
+				String.ReplaceInline(TEXT("\""), TEXT("\"\""));
+				String = FString::Printf(TEXT("\"%s\""), *String);
+				FTCHARToUTF16 UTF16String(*String);
+				StringBuilder.Append(UTF16String.Get(), UTF16String.Length());
+			}
+			else
+			{
+				FString String = FString::Printf(TEXT("\"%s\""), Str);
+				FTCHARToUTF16 UTF16String(*String);
+				StringBuilder.Append(UTF16String.Get(), UTF16String.Length());
+			}
+		}
+		else
+		{
+			FTCHARToUTF16 UTF16String(Str);
+			StringBuilder.Append(UTF16String.Get(), UTF16String.Length());
+		}
+	}
+
+	////////////////////////////////////////////////////////////////////////////////////////////////////
+
+	void FMemoryExporter::FUtf16Writer::Append(const FString& InString)
+	{
+		FTCHARToUTF16 UTF16String(*InString, InString.Len());
+		StringBuilder.Append(UTF16String.Get(), UTF16String.Length());
+	}
+
+}
\ No newline at end of file
diff --git a/Engine/Source/Developer/TraceInsights/Private/Insights/MemoryProfiler/ViewModels/MemoryExporter.h b/Engine/Source/Developer/TraceInsights/Private/Insights/MemoryProfiler/ViewModels/MemoryExporter.h
new file mode 100644
--- /dev/null
+++ b/Engine/Source/Developer/TraceInsights/Private/Insights/MemoryProfiler/ViewModels/MemoryExporter.h
@@ -0,0 +1,154 @@
+// Copyright Epic Games, Inc. All Rights Reserved.
+
+#pragma once
+
+#include "CoreTypes.h"
+#include "Containers/Array.h"
+#include "Containers/Map.h"
+#include "Containers/Set.h"
+#include "HAL/Platform.h"
+#include "Templates/Function.h"
+#include "Templates/UniquePtr.h"
+#include "UObject/NameTypes.h"
+
+// TraceServices
+#include "Insights/MemoryProfiler/ViewModels/MemAllocTable.h"
+#include "TraceServices/Model/AllocationsProvider.h"
+#include "TraceServices/Model/AnalysisSession.h"
+#include "TraceServices/Model/Callstack.h"
+
+// Standard Library
+#include <limits>
+
+// Forward Declarations
+class IFileHandle;
+namespace TraceServices
+{
+	class FMetadataSchema;
+	class IDefinitionProvider;
+	class IMetadataProvider;
+}
+
+
+namespace UE::Insights::MemoryProfiler
+{
+struct FExportMemoryLeakSnapshotParams;
+class FMemoryExporter
+{
+private:
+	static constexpr int32 StringBuilderBufferSize = 64 * 1024;
+public:
+	/**
+	* The list of parameters to run memory leak snapshot
+	*/
+	struct FExportMemoryLeakSnapshotParams
+	{
+		FString OutputPathName;
+		FString Columns;
+		double TimeA, TimeB, TimeC;
+		int32 MaxResults;
+	};
+	// Just like TimeExporter.h, we created a writer class
+	// Why not UTF8?: SMemAllocTableTreeView.cpp uses UTF16 and decided to comply with the standard. 
+	// We made it public though, so other future cmd line exporters can use this helper
+	class FUtf16Writer
+	{
+	public:
+		FUtf16Writer(IFileHandle* InFileHandle, bool bIsCSV)
+			: FileHandle(InFileHandle)
+			, Separator(bIsCSV ? UTF16CHAR(',') : UTF16CHAR('\t'))
+		{
+			check(FileHandle);
+			WriteBOM();
+		}
+
+		~FUtf16Writer() = default;
+
+		TStringBuilder<StringBuilderBufferSize>& GetStringBuilder() { return StringBuilder; }
+		UTF16CHAR GetSeparator() const { return Separator; }
+		UTF16CHAR GetLineEnd() const { return UTF16CHAR('\n'); }
+
+		void AppendSeparator()
+		{
+			StringBuilder.AppendChar(Separator);
+		}
+
+		void AppendLineEnd()
+		{
+			StringBuilder.AppendChar(UTF16CHAR('\n'));
+			WriteStringBuilder(StringBuilderBufferSize - 1024);
+		}
+
+		void AppendString(const TCHAR* InString);
+		void Append(const FString& InString);
+
+		void Flush()
+		{
+			WriteStringBuilder(0);
+		}
+
+	private:
+		void WriteBOM();
+		void WriteStringBuilder(int32 CacheLen);
+
+	private:
+		IFileHandle* FileHandle;
+		TStringBuilder<StringBuilderBufferSize> StringBuilder;
+		UTF16CHAR Separator;
+
+	};
+
+private:
+	//////////////////////////////////////////////////////////////////////
+	// Utilities
+	void ExportMemorySnapshot_InitColumns();
+	bool QueryAllocations(
+		const FExportMemoryLeakSnapshotParams& Params,
+		TArray<const TraceServices::IAllocationsProvider::FAllocation*>& OutAllocations);
+	bool WriteSnapshotCSV(const FString& FilePath,
+		const TArray<const TraceServices::IAllocationsProvider::FAllocation*>& Allocations,
+		const TArray<FName>& ExportColumns);
+	void MakeExportMemoryLeakSnapshotColumnList(const FString& InColumnsString, TArray<FName>& OutColumnList);
+	void InitializeAvailableColumns();
+	FString GetTopFunction(const TraceServices::FCallstack* Callstack);
+	FString GetTopSourceFile(const TraceServices::FCallstack* Callstack);
+	FString FormatCallstackForCSV(const TraceServices::FCallstack* Callstack);
+	void WriteColumnValue(FUtf16Writer& Writer, const FName& ColumnId,
+		const TraceServices::IAllocationsProvider::FAllocation* Alloc,
+		const TraceServices::IAllocationsProvider* AllocationsProvider,
+		const TraceServices::FCallstack* AllocCallstack,
+		const TraceServices::FCallstack* FreeCallstack,
+		const FString& AssetName, const FString& ClassName, const FString& PackageName, uint64 PageSize);
+	void ExtractMetadata(const TraceServices::IAllocationsProvider::FAllocation* Allocation, 
+		uint16 AssetMetadataType,
+		const TraceServices::IMetadataProvider* MetadataProvider,
+		const TraceServices::FMetadataSchema* MetadataSchema,
+		const TraceServices::IDefinitionProvider* DefinitionProvider,
+		FString& OutAssetName, FString& OutClassName, FString& OutPackageName);
+
+
+public:
+	FMemoryExporter(const TraceServices::IAnalysisSession& InSession);
+	virtual ~FMemoryExporter();
+
+	//////////////////////////////////////////////////////////////////////
+	// Exporters
+	int32 ExportMemoryLeakSnapshotAsText(FExportMemoryLeakSnapshotParams& Params);
+
+	//////////////////////////////////////////////////////////////////////
+	// Utilities
+	bool FindBookmarkByName(const FString& BookmarkName, double& OutTime) const;
+
+private:
+	const TraceServices::IAnalysisSession& Session;
+
+	TArray<FName> AvailableColumns;
+	TArray<FName> SnapshotDefaultColumns;
+	
+	static const FName ExportMemoryLeakSnapshot_AllocCallstack;
+	static const FName ExportMemoryLeakSnapshot_FreeCallstack;
+
+	
+};
+} // namespace UE::Insights::MemoryProfiler
+
