[Backout] - CL49971562
#rnx
Original CL Desc
-----------------------------------------------------------------
Automation for a new Memory Leak Exporter (AaBCf rule) in background. It generates a snapshot csv ready for analysis.
Example command usage: UnrealInsights.exe -OpenTraceFile="atracefile.utrace" -unattended -autoquit -noui -nullrhi -ExecOnAnalysisCompleteCmd="[at]=ExportCommands.rsp" -log -AbsLog="MemoryLeakSnapshot.log"
ExportCommands.rsp example: MemoryInsights.ExportLeaks -Rule=AaBCf -Output=MemoryLeakSnapshot.csv -TimeA=0 -BookmarkB="bookmark_name" -BookmarkC="bookmark_name" -MaxResults="0" (zero means all query results)

Although this change only includes the AaBCf rule, it supports future development of other memory exporter rules.

#tests Locally, running tests via command line

#Synced-CL 49554259

#Synced-CL 49554259

#Synced-CL 49554259

diff --git a/Engine/Source/Developer/TraceInsights/Private/Insights/MemoryProfiler/MemoryExportCommandParser.cpp b/Engine/Source/Developer/TraceInsights/Private/Insights/MemoryProfiler/MemoryExportCommandParser.cpp
deleted file mode 100644
diff --git a/Engine/Source/Developer/TraceInsights/Private/Insights/MemoryProfiler/MemoryExportCommandParser.h b/Engine/Source/Developer/TraceInsights/Private/Insights/MemoryProfiler/MemoryExportCommandParser.h
deleted file mode 100644
diff --git a/Engine/Source/Developer/TraceInsights/Private/Insights/MemoryProfiler/MemoryProfilerManager.cpp b/Engine/Source/Developer/TraceInsights/Private/Insights/MemoryProfiler/MemoryProfilerManager.cpp
--- a/Engine/Source/Developer/TraceInsights/Private/Insights/MemoryProfiler/MemoryProfilerManager.cpp
+++ b/Engine/Source/Developer/TraceInsights/Private/Insights/MemoryProfiler/MemoryProfilerManager.cpp
@@ -25,10 +25,6 @@
 #include "Insights/MemoryProfiler/Widgets/SMemoryProfilerWindow.h"
 #include "Insights/Widgets/STimingView.h"
 
-// Commad Line Exporters
-#include "MemoryExportCommandParser.h"
-#include "ViewModels/MemoryExporter.h"
-
 #define LOCTEXT_NAMESPACE "UE::Insights::MemoryProfiler"
 
 // For Memory Insights in editor preview.
@@ -455,139 +451,6 @@ void FMemoryProfilerManager::OnWindowClosedEvent()
 
 ////////////////////////////////////////////////////////////////////////////////////////////////////
 
-typedef TraceServices::IAllocationsProvider::EQueryRule QueryRule;
-static const TSet<QueryRule> RulesUsingTimeA =
-{
-	QueryRule::aAf, QueryRule::afA, QueryRule::Aaf, QueryRule::aAfB,
-	QueryRule::AaBf, QueryRule::aAfaBf, QueryRule::AfB, QueryRule::AaB,
-	QueryRule::AafB, QueryRule::aABf, QueryRule::AaBCf, QueryRule::AaBfC,
-	QueryRule::aABfC, QueryRule::AaBCfD, QueryRule::AoB, QueryRule::AiB
-}; // TimeA is used in all rules
-static const TSet<QueryRule> RulesUsingTimeB =
-{
-	QueryRule::aAfB, QueryRule::AaBf, QueryRule::aAfaBf, QueryRule::AfB, QueryRule::AaB,
-	QueryRule::AafB, QueryRule::aABf, QueryRule::AaBCf, QueryRule::AaBfC, QueryRule::aABfC,
-	QueryRule::AaBCfD, QueryRule::AoB, QueryRule::AiB
-};
-static const TSet<QueryRule> RulesUsingTimeC =
-{
-	QueryRule::AaBCf, QueryRule::AaBfC, QueryRule::aABfC, QueryRule::AaBCfD
-};
-static const TSet<QueryRule> RulesUsingTimeD =
-{
-	QueryRule::AaBCfD
-};
-static const TMap<FString, QueryRule> QueryRuleMap =
-{
-	{ TEXT("aAf"),     QueryRule::aAf },
-	{ TEXT("afA"),     QueryRule::afA },
-	{ TEXT("Aaf"),     QueryRule::Aaf },
-	{ TEXT("aAfB"),    QueryRule::aAfB },
-	{ TEXT("AaBf"),    QueryRule::AaBf },
-	{ TEXT("aAfaBf"),  QueryRule::aAfaBf },
-	{ TEXT("AfB"),     QueryRule::AfB },
-	{ TEXT("AaB"),     QueryRule::AaB },
-	{ TEXT("AafB"),    QueryRule::AafB },
-	{ TEXT("aABf"),    QueryRule::aABf },
-	{ TEXT("AaBCf"),   QueryRule::AaBCf },
-	{ TEXT("AaBfC"),   QueryRule::AaBfC },
-	{ TEXT("aABfC"),   QueryRule::aABfC },
-	{ TEXT("AaBCfD"),  QueryRule::AaBCfD },
-	{ TEXT("AoB"),     QueryRule::AoB },
-	{ TEXT("AiB"),     QueryRule::AiB }
-};
-
-static bool ValidateTimeParameters(QueryRule RuleEnum,
-	const TSet<QueryRule>& RulesUsingTime,
-	const FString& Bookmark,
-	double& Time,
-	const TCHAR* TimeName,
-	const FMemoryExporter& Exporter)
-{
-	return false;
-}
-
-bool FMemoryProfilerManager::Exec(const TCHAR* Cmd, FOutputDevice& Ar)
-{
-	// Parameterized Memory Leak Export Command
-	if (FParse::Command(&Cmd, TEXT("MemoryInsights.ExportAllocs")))
-	{
-		TSharedPtr<FInsightsManager> InsightsManager = FInsightsManager::Get();
-		if (!InsightsManager.IsValid())
-		{
-			UE_LOG(LogMemoryProfiler, Error, TEXT("InsightsManager is invalid"));
-			return false;
-		}
-
-		if (!InsightsManager->GetSession().IsValid())
-		{
-			UE_LOG(LogMemoryProfiler, Error, TEXT("Session is invalid"));
-			return false;
-		}
-
-		FMemoryExporter Exporter(*InsightsManager->GetSession().Get());
-
-		FMemoryExportCommandParser ExportCommandParser;
-		if (!ExportCommandParser.ParseParameters(Cmd, Ar))
-		{
-			UE_LOG(LogMemoryProfiler, Error, TEXT("Exporter stopped while parsing cmd parameters"));
-			return false;
-		}
-
-		if (!ExportCommandParser.ValidateParameters(Exporter))
-		{
-			UE_LOG(LogMemoryProfiler, Error, TEXT("Exporter stopped while validating parameters"));
-			return false;
-		}
-
-		// Below logic corresponds to Rule=AaBCf
-		if(ExportCommandParser.GetRuleEnum() == TraceServices::IAllocationsProvider::EQueryRule::AaBCf)
-		{
-			UE_LOG(LogMemoryProfiler, Log, TEXT(" === Memory Leak Exporter. Rule: \"%s\" === "), *ExportCommandParser.GetRuleName());
-			FMemoryExporter::FExportMemoryLeakSnapshotParams SnapshotParams;
-			SnapshotParams.OutputPathName = ExportCommandParser.GetOutputPathName();
-			SnapshotParams.TimeA = ExportCommandParser.GetTimeA();
-			SnapshotParams.TimeB = ExportCommandParser.GetTimeB();
-			SnapshotParams.TimeC = ExportCommandParser.GetTimeC();
-			SnapshotParams.MaxResults = ExportCommandParser.GetMaxResults();
-			SnapshotParams.Columns = ExportCommandParser.GetColumns();
-
-			FString BookmarkA = ExportCommandParser.GetBookmarkA();
-			FString BookmarkB = ExportCommandParser.GetBookmarkB();
-			FString BookmarkC = ExportCommandParser.GetBookmarkC();
-
-			UE_LOG(LogMemoryProfiler, Log, TEXT("Parameters provided by user:"));
-			UE_LOG(LogMemoryProfiler, Log, TEXT("-> Output:     %s"), *SnapshotParams.OutputPathName);
-			UE_LOG(LogMemoryProfiler, Log, TEXT("-> TimeA:      %.2f%s"), SnapshotParams.TimeA, BookmarkA.IsEmpty() ? TEXT("") : *FString::Printf(TEXT(" (from bookmark '%s')"), *BookmarkA));
-			UE_LOG(LogMemoryProfiler, Log, TEXT("-> TimeB:      %.2f%s"), SnapshotParams.TimeB, BookmarkB.IsEmpty() ? TEXT("") : *FString::Printf(TEXT(" (from bookmark '%s')"), *BookmarkB));
-			UE_LOG(LogMemoryProfiler, Log, TEXT("-> TimeC:      %.2f%s"), SnapshotParams.TimeC, BookmarkC.IsEmpty() ? TEXT("") : *FString::Printf(TEXT(" (from bookmark '%s')"), *BookmarkC));
-			if (ExportCommandParser.GetTimeD() >= 0.0 || !ExportCommandParser.GetBookmarkD().IsEmpty())
-			{
-				UE_LOG(LogMemoryProfiler, Warning, TEXT("-> TimeD SKIPPED. Not required in this rule"));
-			}
-			UE_LOG(LogMemoryProfiler, Log, TEXT("-> MaxResults: %d (0 = unlimited)"), SnapshotParams.MaxResults);
-			UE_LOG(LogMemoryProfiler, Log, TEXT("Starting export..."));
-			int ResultCount = Exporter.ExportMemoryLeakSnapshotAsText(SnapshotParams);
-
-			if (ResultCount >= 0)
-			{
-				UE_LOG(LogMemoryProfiler, Log, TEXT("SUCCESS! Exported %d allocations"), ResultCount);
-				UE_LOG(LogMemoryProfiler, Log, TEXT("Output file: %s"), *SnapshotParams.OutputPathName);
-				return true;
-			}
-			else
-			{
-				UE_LOG(LogMemoryProfiler, Error, TEXT("Export failed with code %d"), ResultCount);
-				return false;
-			}
-		}
-		
-	}
-	return false;
-}
-
-////////////////////////////////////////////////////////////////////////////////////////////////////
-
 } // namespace UE::Insights::MemoryProfiler
 
 #undef LOCTEXT_NAMESPACE
diff --git a/Engine/Source/Developer/TraceInsights/Private/Insights/MemoryProfiler/MemoryProfilerManager.h b/Engine/Source/Developer/TraceInsights/Private/Insights/MemoryProfiler/MemoryProfilerManager.h
--- a/Engine/Source/Developer/TraceInsights/Private/Insights/MemoryProfiler/MemoryProfilerManager.h
+++ b/Engine/Source/Developer/TraceInsights/Private/Insights/MemoryProfiler/MemoryProfilerManager.h
@@ -55,7 +55,6 @@ public:
 	virtual void Shutdown() override;
 	virtual void RegisterMajorTabs(IUnrealInsightsModule& InsightsModule) override;
 	virtual void UnregisterMajorTabs() override;
-	virtual bool Exec(const TCHAR* Cmd, FOutputDevice& Ar) override;
 	virtual void OnWindowClosedEvent() override;
 
 	//////////////////////////////////////////////////
diff --git a/Engine/Source/Developer/TraceInsights/Private/Insights/MemoryProfiler/ViewModels/MemoryExporter.cpp b/Engine/Source/Developer/TraceInsights/Private/Insights/MemoryProfiler/ViewModels/MemoryExporter.cpp
deleted file mode 100644
diff --git a/Engine/Source/Developer/TraceInsights/Private/Insights/MemoryProfiler/ViewModels/MemoryExporter.h b/Engine/Source/Developer/TraceInsights/Private/Insights/MemoryProfiler/ViewModels/MemoryExporter.h
deleted file mode 100644
