Add format string to timing event display name even when the event has no other metadata.
Add a new setting that adds the format to the timer names in the Timers/Callers/Callees panel when enabled.

#Synced-CL 49071325

diff --git a/Engine/Config/BaseUnrealInsightsSettings.ini b/Engine/Config/BaseUnrealInsightsSettings.ini
--- a/Engine/Config/BaseUnrealInsightsSettings.ini
+++ b/Engine/Config/BaseUnrealInsightsSettings.ini
@@ -53,6 +53,10 @@ AutoScrollMinDelay=0.3
 ; This option can be changed from the Settings sub-menu in the Context Menu of any Timer in the Timers Panel.
 ThreadsToUseRatio=0.75
 
+; When enabled, the format will be appended to the timers static name for events with metadata specs.
+; This option can be changed from the Settings sub-menu in the Context Menu of any Timer in the Timers Panel.
+bShowFormatInTimerNames=True
+
 [Insights.TimingProfiler.FramesView]
 ; These options can be changed from the context menu of the Frames view.
 
diff --git a/Engine/Source/Developer/TraceInsights/Private/Insights/InsightsSettings.cpp b/Engine/Source/Developer/TraceInsights/Private/Insights/InsightsSettings.cpp
--- a/Engine/Source/Developer/TraceInsights/Private/Insights/InsightsSettings.cpp
+++ b/Engine/Source/Developer/TraceInsights/Private/Insights/InsightsSettings.cpp
@@ -63,6 +63,7 @@ void FInsightsSettings::ResetToDefaults()
 	AutoScrollViewportOffsetPercent = Defaults.AutoScrollViewportOffsetPercent;
 	AutoScrollMinDelay = Defaults.AutoScrollMinDelay;
 	ThreadsToUseRatio = Defaults.ThreadsToUseRatio;
+	bShowFormatInTimerNames = Defaults.bShowFormatInTimerNames;
 }
 
 ////////////////////////////////////////////////////////////////////////////////////////////////////
@@ -122,6 +123,7 @@ void FInsightsSettings::LoadFromConfig()
 	GConfig->GetBool(TEXT("Insights.TimingProfiler"), TEXT("bAutoHideEmptyTracks"), bAutoHideEmptyTracks, SettingsIni);
 	GConfig->GetBool(TEXT("Insights.TimingProfiler"), TEXT("bAllowPanningOnScreenEdges"), bAllowPanningOnScreenEdges, SettingsIni);
 	GConfig->GetDouble(TEXT("Insights.TimingProfiler"), TEXT("ThreadsToUseRatio"), ThreadsToUseRatio, SettingsIni);
+	GConfig->GetBool(TEXT("Insights.TimingProfiler"), TEXT("bShowFormatInTimerNames"), bShowFormatInTimerNames, SettingsIni);
 	GConfig->GetInt(TEXT("Insights.TimingProfiler"), TEXT("TimingTrackDepthLimit.Min"), TimingTrackDepthLimitMin, SettingsIni);
 	GConfig->GetInt(TEXT("Insights.TimingProfiler"), TEXT("TimingTrackDepthLimit.Max"), TimingTrackDepthLimitMax, SettingsIni);
 
@@ -300,6 +302,7 @@ void FInsightsSettings::SaveToConfig()
 	GConfig->SetBool(TEXT("Insights.TimingProfiler"), TEXT("bAutoHideEmptyTracks"), bAutoHideEmptyTracks, SettingsIni);
 	GConfig->SetBool(TEXT("Insights.TimingProfiler"), TEXT("bAllowPanningOnScreenEdges"), bAllowPanningOnScreenEdges, SettingsIni);
 	GConfig->SetDouble(TEXT("Insights.TimingProfiler"), TEXT("ThreadsToUseRatio"), ThreadsToUseRatio, SettingsIni);
+	GConfig->SetBool(TEXT("Insights.TimingProfiler"), TEXT("bShowFormatInTimerNames"), bShowFormatInTimerNames, SettingsIni);
 	GConfig->SetInt(TEXT("Insights.TimingProfiler"), TEXT("TimingTrackDepthLimit.Min"), TimingTrackDepthLimitMin, SettingsIni);
 	GConfig->SetInt(TEXT("Insights.TimingProfiler"), TEXT("TimingTrackDepthLimit.Max"), TimingTrackDepthLimitMax, SettingsIni);
 
diff --git a/Engine/Source/Developer/TraceInsights/Private/Insights/InsightsSettings.h b/Engine/Source/Developer/TraceInsights/Private/Insights/InsightsSettings.h
--- a/Engine/Source/Developer/TraceInsights/Private/Insights/InsightsSettings.h
+++ b/Engine/Source/Developer/TraceInsights/Private/Insights/InsightsSettings.h
@@ -370,6 +370,10 @@ public:
 	void SetThreadRatio(double Ratio) { ThreadsToUseRatio = Ratio; }
 	void SetAndSaveThreadRatio(double Ratio) { SET_AND_SAVE(ThreadsToUseRatio, Ratio); }
 
+	bool GetShowFormatInTimerNames() const { return bShowFormatInTimerNames; }
+	void SetShowFormatInTimerNames(bool Value) { bShowFormatInTimerNames = Value; }
+	void SetAndSaveShowFormatInTimerNames(bool Value) { SET_AND_SAVE(bShowFormatInTimerNames, Value); }
+
 	//////////////////////////////////////////////////
 	// [Insights.TimingProfiler.FramesView]
 
@@ -582,6 +586,9 @@ private:
 	/** The ratio of available threads to use when performing per-frame aggregations. */
 	double ThreadsToUseRatio = 0.75;
 
+	/** When enabled, the format will be appended to the timers static name for events with metadata specs. */
+	bool bShowFormatInTimerNames = true;
+
 	//////////////////////////////////////////////////
 	// [Insights.TimingProfiler.FramesView]
 
diff --git a/Engine/Source/Developer/TraceInsights/Private/Insights/TimingProfiler/TimingProfilerManager.cpp b/Engine/Source/Developer/TraceInsights/Private/Insights/TimingProfiler/TimingProfilerManager.cpp
--- a/Engine/Source/Developer/TraceInsights/Private/Insights/TimingProfiler/TimingProfilerManager.cpp
+++ b/Engine/Source/Developer/TraceInsights/Private/Insights/TimingProfiler/TimingProfilerManager.cpp
@@ -738,6 +738,23 @@ bool FTimingProfilerManager::Exec(const TCHAR* Cmd, FOutputDevice& Ar)
 		Filename.TrimQuotesInline();
 		Ar.Logf(TEXT("  Filename: \"%s\""), *Filename);
 
+		while (Cmd && Cmd[0] != TEXT('\0'))
+		{
+			FString Token;
+			if (FParse::Token(Cmd, Token, bUseEscape))
+			{
+				Ar.Logf(TEXT("  Token: %s"), *Token);
+
+				static constexpr TCHAR ShowFormatInTimerNamesToken[] = TEXT("-ShowFormatInTimerNames=");
+				if (Token.StartsWith(ShowFormatInTimerNamesToken))
+				{
+					Token.RightChopInline(UE_ARRAY_COUNT(ShowFormatInTimerNamesToken) - 1);
+					Token.TrimQuotesInline();
+					Params.bShowFormatInTimerNames = Token.ToBool();
+				}
+			}
+		}
+
 		Exporter.ExportTimersAsText(Filename, Params);
 		return true;
 	}
@@ -1034,6 +1051,7 @@ bool FTimingProfilerManager::Exec(const TCHAR* Cmd, FOutputDevice& Ar)
 				static constexpr TCHAR StartTimeToken[] = TEXT("-startTime=");
 				static constexpr TCHAR EndTimeToken[] = TEXT("-endTime=");
 				static constexpr TCHAR RegionToken[] = TEXT("-region=");
+				static constexpr TCHAR ShowFormatInTimerNamesToken[] = TEXT("-ShowFormatInTimerNames=");
 
 				if (Token.StartsWith(ThreadsToken))
 				{
@@ -1078,6 +1096,14 @@ bool FTimingProfilerManager::Exec(const TCHAR* Cmd, FOutputDevice& Ar)
 					Token.TrimQuotesInline();
 					Params.Region = TCHAR_TO_ANSI(*Token);
 				}
+				else if (Token.StartsWith(ShowFormatInTimerNamesToken))
+				{
+					// Default: -ShowFormatInTimerNames=
+					// Example: -ShowFormatInTimerNames="False"
+					Token.RightChopInline(UE_ARRAY_COUNT(ShowFormatInTimerNamesToken) - 1);
+					Token.TrimQuotesInline();
+					Params.bShowFormatInTimerNames = Token.ToBool();
+				}
 				else
 				{
 					Ar.Logf(ELogVerbosity::Warning, TEXT("Unknown Cmd Param: %s"), *Token);
diff --git a/Engine/Source/Developer/TraceInsights/Private/Insights/TimingProfiler/Tracks/ThreadTimingTrack.cpp b/Engine/Source/Developer/TraceInsights/Private/Insights/TimingProfiler/Tracks/ThreadTimingTrack.cpp
--- a/Engine/Source/Developer/TraceInsights/Private/Insights/TimingProfiler/Tracks/ThreadTimingTrack.cpp
+++ b/Engine/Source/Developer/TraceInsights/Private/Insights/TimingProfiler/Tracks/ThreadTimingTrack.cpp
@@ -449,6 +449,11 @@ static void AppendMetadataToString(FString& InOutStr, TArrayView<const uint8>& I
 
 static void AppendMetadataToString(FString& InOutStr, const FString& InFormat, TArrayView<const uint8>& InMetadata)
 {
+	if (InFormat.IsEmpty())
+	{
+		return;
+	}
+
 	FMemoryReaderView MemoryReader(InMetadata);
 	FCborReader CborReader(&MemoryReader, ECborEndianness::StandardCompliant);
 	FCborContext Context;
@@ -461,6 +466,8 @@ static void AppendMetadataToString(FString& InOutStr, const FString& InFormat, T
 		Format.RightChopInline(2);
 	}
 
+	InOutStr += TEXTVIEW(" - ");
+
 	const FString Specifiers = TEXT("diuoxXfFeEgGaAcspn");
 	auto GetNextFormatSection = [&Format, &Specifiers]()
 	{
@@ -497,11 +504,6 @@ static void AppendMetadataToString(FString& InOutStr, const FString& InFormat, T
 			break;
 		}
 
-		if (Index == 0)
-		{
-			InOutStr += TEXTVIEW(" - ");
-		}
-
 		constexpr int MaxLength = 1024;
 		TCHAR Data[MaxLength];
 
@@ -584,6 +586,41 @@ static void AppendMetadataToString(FString& InOutStr, const FString& InFormat, T
 
 ////////////////////////////////////////////////////////////////////////////////////////////////////
 
+void AppendInfoToTimingEventDisplayName(
+	FString& InOutDisplayName,
+	double EventStartTime,
+	double EventEndTime,
+	uint32 TimerIndex,
+	const TraceServices::FTimingProfilerTimer* Timer,
+	const TraceServices::ITimingProfilerProvider& TimingProfilerProvider,
+	const TraceServices::ITimingProfilerTimerReader& TimerReader)
+{
+	const double Duration = EventEndTime - EventStartTime;
+	FTimingEventsTrackDrawStateBuilder::AppendDurationToEventName(InOutDisplayName, Duration);
+
+	TArrayView<const uint8> Metadata;
+	if (int32(TimerIndex) < 0) // has static metadata
+	{
+		Metadata = TimerReader.GetMetadata(TimerIndex);
+	}
+	const TraceServices::FMetadataSpec* MetadataSpec = nullptr;
+	if (Timer->HasValidMetadataSpecId()) // dynamic metadata with spec
+	{
+		MetadataSpec = TimingProfilerProvider.GetMetadataSpec(Timer->MetadataSpecId);
+		if (MetadataSpec)
+		{
+			AppendMetadataToString(InOutDisplayName, MetadataSpec->Format, Metadata);
+			return;
+		}
+	}
+	if (Metadata.Num() > 0)
+	{
+		AppendMetadataToString(InOutDisplayName, Metadata);
+	}
+}
+
+////////////////////////////////////////////////////////////////////////////////////////////////////
+
 void AddTimingEventToBuilder(
 	ITimingEventsTrackDrawStateBuilder& Builder,
 	double EventStartTime,
@@ -635,35 +672,7 @@ void AddTimingEventToBuilder(
 			const float MinWidth = static_cast<float>(EventName.Len()) * 4.0f + 32.0f;
 			if (Width > MinWidth)
 			{
-				//EventName = TEXT("*") + EventName; // for debugging
-
-				const double Duration = EventEndTime - EventStartTime;
-				FTimingEventsTrackDrawStateBuilder::AppendDurationToEventName(EventName, Duration);
-
-				if (int32(TimerIndex) < 0) // has metadata?
-				{
-					//EventName = TEXT("!") + EventName; // for debugging
-
-					TArrayView<const uint8> Metadata = TimerReader.GetMetadata(TimerIndex);
-
-					if (Metadata.Num() > 0)
-					{
-						const TraceServices::FMetadataSpec* MetadataSpec = nullptr;
-						if (Timer->HasValidMetadataSpecId())
-						{
-							MetadataSpec = TimingProfilerProvider.GetMetadataSpec(Timer->MetadataSpecId);
-						}
-
-						if (MetadataSpec)
-						{
-							AppendMetadataToString(EventName, MetadataSpec->Format, Metadata);
-						}
-						else
-						{
-							AppendMetadataToString(EventName, Metadata);
-						}
-					}
-				}
+				AppendInfoToTimingEventDisplayName(EventName, EventStartTime, EventEndTime, TimerIndex, Timer, TimingProfilerProvider, TimerReader);
 			}
 
 			return EventName;
@@ -1321,31 +1330,7 @@ void FThreadTimingTrackImpl::OnClipboardCopyEvent(const ITimingEvent& InSelected
 				{
 					FString EventName(TimerPtr->Name);
 
-					FTimingEventsTrackDrawStateBuilder::AppendDurationToEventName(EventName, TrackEvent.GetDuration());
-
-					const uint32 TimerIndex = TrackEvent.GetTimerIndex();
-					if (int32(TimerIndex) < 0) // has metadata?
-					{
-						TArrayView<const uint8> Metadata = TimerReader.GetMetadata(TimerIndex);
-
-						const TraceServices::FMetadataSpec* MetadataSpec = nullptr;
-						if (TimerPtr && TimerPtr->HasValidMetadataSpecId())
-						{
-							MetadataSpec = TimingProfilerProvider.GetMetadataSpec(TimerPtr->MetadataSpecId);
-						}
-
-						if (Metadata.Num() > 0)
-						{
-							if (MetadataSpec)
-							{
-								AppendMetadataToString(EventName, MetadataSpec->Format, Metadata);
-							}
-							else
-							{
-								AppendMetadataToString(EventName, Metadata);
-							}
-						}
-					}
+					AppendInfoToTimingEventDisplayName(EventName, TrackEvent.GetStartTime(), TrackEvent.GetEndTime(), TrackEvent.GetTimerIndex(), TimerPtr, TimingProfilerProvider, TimerReader);
 
 					// Copy name of selected timing event to clipboard.
 					FPlatformApplicationMisc::ClipboardCopy(*EventName);
diff --git a/Engine/Source/Developer/TraceInsights/Private/Insights/TimingProfiler/ViewModels/TimingExporter.cpp b/Engine/Source/Developer/TraceInsights/Private/Insights/TimingProfiler/ViewModels/TimingExporter.cpp
--- a/Engine/Source/Developer/TraceInsights/Private/Insights/TimingProfiler/ViewModels/TimingExporter.cpp
+++ b/Engine/Source/Developer/TraceInsights/Private/Insights/TimingProfiler/ViewModels/TimingExporter.cpp
@@ -297,7 +297,7 @@ int32 FTimingExporter::ExportTimersAsText(const FString& Filename, FExportTimers
 		TraceServices::FAnalysisSessionReadScope SessionReadScope(Session);
 
 		TimingProfilerProvider->ReadTimers(
-			[&Writer, &StringBuilder, &TimerCount](const TraceServices::ITimingProfilerTimerReader& TimerReader)
+			[&Writer, &StringBuilder, &TimerCount, &TimingProfilerProvider, &Params](const TraceServices::ITimingProfilerTimerReader& TimerReader)
 			{
 				TimerCount = TimerReader.GetTimerCount();
 				for (uint32 TimerIndex = 0; TimerIndex < TimerCount; ++TimerIndex)
@@ -324,7 +324,16 @@ int32 FTimingExporter::ExportTimersAsText(const FString& Filename, FExportTimers
 						break;
 					}
 					Writer.AppendSeparator();
-					Writer.AppendString(Timer.Name);
+					if (Params.bShowFormatInTimerNames)
+					{
+						TStringBuilder<1024> NameStringBuilder;
+						TimingProfilerProvider->MakeTimerDisplayName(Timer, NameStringBuilder);
+						Writer.AppendString(*NameStringBuilder);
+					}
+					else
+					{
+						Writer.AppendString(Timer.Name);
+					}
 					Writer.AppendSeparator();
 					if (Timer.File)
 					{
@@ -1174,6 +1183,7 @@ int32 FTimingExporter::ExportTimerCalleesAsText(const FString& Filename, const F
 	// Write rows
 	{
 		FUtf8StringBuilder& StringBuilder = Writer.GetStringBuilder();
+		TraceServices::FAnalysisSessionReadScope SessionReadScope(Session);
 
 		for (uint32 TimerId : Params.TimerIds)
 		{
@@ -1193,7 +1203,16 @@ int32 FTimingExporter::ExportTimerCalleesAsText(const FString& Filename, const F
 				Writer.AppendSeparator();
 				StringBuilder.Appendf(UTF8TEXT("%u"), CurrentNode->Parent ? CurrentNode->Parent->Timer->Id : uint32(-1));
 				Writer.AppendSeparator();
-				Writer.AppendString(CurrentNode->Timer->Name);
+				if (Params.bShowFormatInTimerNames)
+				{
+					TStringBuilder<1024> NameStringBuilder;
+					TimingProfilerProvider->MakeTimerDisplayName(*CurrentNode->Timer, NameStringBuilder);
+					Writer.AppendString(*NameStringBuilder);
+				}
+				else
+				{
+					Writer.AppendString(CurrentNode->Timer->Name);
+				}
 				Writer.AppendSeparator();
 				StringBuilder.Appendf(UTF8TEXT("%llu"), CurrentNode->Count);
 				Writer.AppendSeparator();
diff --git a/Engine/Source/Developer/TraceInsights/Private/Insights/TimingProfiler/ViewModels/TimingExporter.h b/Engine/Source/Developer/TraceInsights/Private/Insights/TimingProfiler/ViewModels/TimingExporter.h
--- a/Engine/Source/Developer/TraceInsights/Private/Insights/TimingProfiler/ViewModels/TimingExporter.h
+++ b/Engine/Source/Developer/TraceInsights/Private/Insights/TimingProfiler/ViewModels/TimingExporter.h
@@ -59,6 +59,11 @@ public:
 		 * If nullptr, it uses the default list of columns.
 		 */
 		const TArray<FName>* Columns = nullptr;
+
+		/**
+		 *  When enabled, the format will be appended to the timers static name for events with metadata specs.
+		 */
+		bool bShowFormatInTimerNames = true;
 	};
 
 	struct FExportTimingEventsParams
@@ -140,6 +145,11 @@ public:
 		using FExportTimingEventsParams::Region;
 
 		TSet<uint32> TimerIds;
+
+		/**
+		 *  When enabled, the format will be appended to the timers static name for events with metadata specs.
+		 */
+		bool bShowFormatInTimerNames = true;
 	};
 
 	struct FExportCountersParams
diff --git a/Engine/Source/Developer/TraceInsights/Private/Insights/TimingProfiler/Widgets/STimerTreeView.cpp b/Engine/Source/Developer/TraceInsights/Private/Insights/TimingProfiler/Widgets/STimerTreeView.cpp
--- a/Engine/Source/Developer/TraceInsights/Private/Insights/TimingProfiler/Widgets/STimerTreeView.cpp
+++ b/Engine/Source/Developer/TraceInsights/Private/Insights/TimingProfiler/Widgets/STimerTreeView.cpp
@@ -1495,28 +1495,38 @@ void STimerTreeView::Reset()
 void STimerTreeView::SetTree(const TraceServices::FTimingProfilerButterflyNode& Root)
 {
 	TreeNodes.Reset();
+	FTimerNodePtr RootTimerNodePtr;
 
-	FTimerNodePtr RootTimerNodePtr = CreateTimerNodeRec(Root);
-	if (RootTimerNodePtr)
+	TSharedPtr<const TraceServices::IAnalysisSession> Session = FInsightsManager::Get()->GetSession();
+	if (!Session.IsValid())
+	{
+		return;
+	}
+	const TraceServices::ITimingProfilerProvider* TimingProfilerProvider = TraceServices::ReadTimingProfilerProvider(*Session.Get());
+	if (TimingProfilerProvider)
 	{
-		// Mark the hot path. The child nodes are already sorted by InclTime (descending), so we just follow the first child.
-		FTimerNodePtr TimerNodePtr = RootTimerNodePtr;
-		while (TimerNodePtr.IsValid())
+		RootTimerNodePtr = CreateTimerNodeRec(Root, *TimingProfilerProvider);
+		if (RootTimerNodePtr)
 		{
-			TimerNodePtr->SetIsHotPath(true);
-			const TArray<FBaseTreeNodePtr>& Children = TimerNodePtr->GetChildren();
-			if (Children.Num() > 0)
-			{
-				check(Children[0]->Is<FTimerNode>());
-				TimerNodePtr = StaticCastSharedPtr<FTimerNode>(Children[0]);
-			}
-			else
+			// Mark the hot path. The child nodes are already sorted by InclTime (descending), so we just follow the first child.
+			FTimerNodePtr TimerNodePtr = RootTimerNodePtr;
+			while (TimerNodePtr.IsValid())
 			{
-				TimerNodePtr = nullptr;
+				TimerNodePtr->SetIsHotPath(true);
+				const TArray<FBaseTreeNodePtr>& Children = TimerNodePtr->GetChildren();
+				if (Children.Num() > 0)
+				{
+					check(Children[0]->Is<FTimerNode>());
+					TimerNodePtr = StaticCastSharedPtr<FTimerNode>(Children[0]);
+				}
+				else
+				{
+					TimerNodePtr = nullptr;
+				}
 			}
-		}
 
-		TreeNodes.Add(RootTimerNodePtr);
+			TreeNodes.Add(RootTimerNodePtr);
+		}
 	}
 
 	SortTreeNodes();
@@ -1531,7 +1541,7 @@ void STimerTreeView::SetTree(const TraceServices::FTimingProfilerButterflyNode&
 
 ////////////////////////////////////////////////////////////////////////////////////////////////////
 
-FTimerNodePtr STimerTreeView::CreateTimerNodeRec(const TraceServices::FTimingProfilerButterflyNode& Node)
+FTimerNodePtr STimerTreeView::CreateTimerNodeRec(const TraceServices::FTimingProfilerButterflyNode& Node, const TraceServices::ITimingProfilerProvider& TimingProfilerProvider)
 {
 	if (Node.Timer == nullptr)
 	{
@@ -1547,7 +1557,23 @@ FTimerNodePtr STimerTreeView::CreateTimerNodeRec(const TraceServices::FTimingPro
 
 	TSharedPtr<STimingProfilerWindow> ProfilerWindow = FTimingProfilerManager::Get()->GetProfilerWindow();
 	const bool bIsTimerAddedToGraphs = ProfilerWindow.IsValid() ? ProfilerWindow->IsTimerAddedToGraphs(Node.Timer->Id) : false;
-	FTimerNodePtr TimerNodePtr = MakeShared<FTimerNode>(Node.Timer->Id, Node.Timer->Name, Type, true);
+
+	TStringBuilder<1024> StringBuilder;
+	if (FInsightsManager::Get()->GetSettings().GetShowFormatInTimerNames())
+	{
+		TSharedPtr<const TraceServices::IAnalysisSession> Session = FInsightsManager::Get()->GetSession();
+		if (Session)
+		{
+			TraceServices::FAnalysisSessionReadScope SessionReadScope(*Session.Get());
+			TimingProfilerProvider.MakeTimerDisplayName(*Node.Timer, StringBuilder);
+		}
+	}
+	else
+	{
+		StringBuilder.Append(Node.Timer->Name);
+	}
+
+	FTimerNodePtr TimerNodePtr = MakeShared<FTimerNode>(Node.Timer->Id, *StringBuilder, Type, true);
 	FLinearColor ColorOverride;
 
 	if (ProfilerWindow.IsValid() && ProfilerWindow->GetTimerColor(Node.Timer->Id, ColorOverride))
@@ -1576,7 +1602,7 @@ FTimerNodePtr STimerTreeView::CreateTimerNodeRec(const TraceServices::FTimingPro
 	{
 		if (ChildNodePtr != nullptr)
 		{
-			FTimerNodePtr ChildTimerNodePtr = CreateTimerNodeRec(*ChildNodePtr);
+			FTimerNodePtr ChildTimerNodePtr = CreateTimerNodeRec(*ChildNodePtr, TimingProfilerProvider);
 			if (ChildTimerNodePtr)
 			{
 				TimerNodePtr->AddChildAndSetParent(ChildTimerNodePtr);
diff --git a/Engine/Source/Developer/TraceInsights/Private/Insights/TimingProfiler/Widgets/STimerTreeView.h b/Engine/Source/Developer/TraceInsights/Private/Insights/TimingProfiler/Widgets/STimerTreeView.h
--- a/Engine/Source/Developer/TraceInsights/Private/Insights/TimingProfiler/Widgets/STimerTreeView.h
+++ b/Engine/Source/Developer/TraceInsights/Private/Insights/TimingProfiler/Widgets/STimerTreeView.h
@@ -81,7 +81,7 @@ public:
 
 private:
 	void InitCommandList();
-	FTimerNodePtr CreateTimerNodeRec(const TraceServices::FTimingProfilerButterflyNode& Node);
+	FTimerNodePtr CreateTimerNodeRec(const TraceServices::FTimingProfilerButterflyNode& Node, const TraceServices::ITimingProfilerProvider& TimingProfilerProvider);
 	void ExpandNodesRec(FTimerNodePtr NodePtr, int32 Depth);
 
 	////////////////////////////////////////////////////////////////////////////////////////////////////
diff --git a/Engine/Source/Developer/TraceInsights/Private/Insights/TimingProfiler/Widgets/STimersView.cpp b/Engine/Source/Developer/TraceInsights/Private/Insights/TimingProfiler/Widgets/STimersView.cpp
--- a/Engine/Source/Developer/TraceInsights/Private/Insights/TimingProfiler/Widgets/STimersView.cpp
+++ b/Engine/Source/Developer/TraceInsights/Private/Insights/TimingProfiler/Widgets/STimersView.cpp
@@ -1325,37 +1325,60 @@ void STimersView::TreeView_BuildSettingsMenu(FMenuBuilder& MenuBuilder)
 
 	MenuBuilder.BeginSection("Settings");
 	{
-		MenuBuilder.AddEditableText
-		(
-			LOCTEXT("ThreadsAggregation", "Number of Threads"),
-			LOCTEXT("ThreadsAggregation_Tooltip",
-				"Number of threads being used for Aggregation.\n"
-				"Using higher values may cause systems to freeze.\n"
-				"This value is persistent to the UnrealInsightsSettings.ini file."),
-			FSlateIcon(),
-			TAttribute<FText>::Create(TAttribute<FText>::FGetter::CreateRaw(this, &STimersView::GetThreadRatio)),
-			FOnTextCommitted::CreateLambda([this](const FText& NewValue, ETextCommit::Type)
-			{
-				double Ratio = FCString::Atof(*NewValue.ToString());
-				OnThreadRatioCommited(Ratio);
-			})
-		);
+		{
+			MenuBuilder.AddEditableText
+			(
+				LOCTEXT("ThreadsAggregation", "Number of Threads"),
+				LOCTEXT("ThreadsAggregation_Tooltip",
+					"Number of threads being used for Aggregation.\n"
+					"Using higher values may cause systems to freeze.\n"
+					"This value is persistent to the UnrealInsightsSettings.ini file."),
+				FSlateIcon(),
+				TAttribute<FText>::Create(TAttribute<FText>::FGetter::CreateRaw(this, &STimersView::GetThreadRatio)),
+				FOnTextCommitted::CreateLambda([this](const FText& NewValue, ETextCommit::Type)
+					{
+						double Ratio = FCString::Atof(*NewValue.ToString());
+						OnThreadRatioCommited(Ratio);
+					})
+			);
+		}
 
-		FUIAction Action_SetSourceFilesSearchPath
-		(
-			FExecuteAction::CreateSP(this, &STimersView::SetSourceFileSearchPaths_Execute),
-			FCanExecuteAction::CreateLambda([]() { return true; })
-		);
+		{
+			FUIAction Action_SetSourceFilesSearchPath
+			(
+				FExecuteAction::CreateSP(this, &STimersView::SetSourceFileSearchPaths_Execute),
+				FCanExecuteAction::CreateLambda([]() { return true; })
+			);
 
-		MenuBuilder.AddMenuEntry
-		(
-			LOCTEXT("SetSourceFilesSearchPath_Fmt", "Source File Path"),
-			TAttribute<FText>::CreateSP(this, &STimersView::SetSourceFileSearchPaths_GetTooltip),
-			FSlateIcon(),
-			Action_SetSourceFilesSearchPath,
-			NAME_None,
-			EUserInterfaceActionType::Button
-		);
+			MenuBuilder.AddMenuEntry
+			(
+				LOCTEXT("SetSourceFilesSearchPath_Fmt", "Source File Path"),
+				TAttribute<FText>::CreateSP(this, &STimersView::SetSourceFileSearchPaths_GetTooltip),
+				FSlateIcon(),
+				Action_SetSourceFilesSearchPath,
+				NAME_None,
+				EUserInterfaceActionType::Button
+			);
+		}
+
+		{
+			FUIAction Action_ShowFormatInTimerNames
+			(
+				FExecuteAction::CreateSP(this, &STimersView::ToggleShowFormatInTimerNames_Execute),
+				FCanExecuteAction::CreateLambda([]() { return true; }),
+				FIsActionChecked::CreateLambda([]() { return FInsightsManager::Get()->GetSettings().GetShowFormatInTimerNames(); })
+			);
+
+			MenuBuilder.AddMenuEntry
+			(
+				LOCTEXT("SetShowFormatInTimerNames_Fmt", "Show Format In Timer Names"),
+				LOCTEXT("SetShowFormatInTimerNamesTooltip_Fmt", "Append the format to timers names for events with metadata."),
+				FSlateIcon(),
+				Action_ShowFormatInTimerNames,
+				NAME_None,
+				EUserInterfaceActionType::ToggleButton
+			);
+		}
 	}
 	MenuBuilder.EndSection();
 }
@@ -2831,7 +2854,19 @@ void STimersView::RebuildTree(bool bResync)
 					Timer.Type == TraceServices::ETimingProfilerTimerType::CpuScope ? ETimerNodeType::CpuScope :
 					Timer.Type == TraceServices::ETimingProfilerTimerType::CpuSampling ? ETimerNodeType::CpuSampling :
 					ETimerNodeType::InvalidOrMax;
-				FTimerNodePtr TimerNodePtr = MakeShared<FTimerNode>(Timer.Id, Timer.Name, Type, false);
+
+				FTimerNodePtr TimerNodePtr;
+				TStringBuilder<1024> StringBuilder;
+				if (FInsightsManager::Get()->GetSettings().GetShowFormatInTimerNames())
+				{
+					TimingProfilerProvider.MakeTimerDisplayName(Timer, StringBuilder);
+				}
+				else
+				{
+					StringBuilder.Append(Timer.Name);
+				}
+
+				TimerNodePtr = MakeShared<FTimerNode>(Timer.Id, *StringBuilder, Type, false);
 				TimerNodePtr->SetDefaultSortOrder(TimerIndex + 1);
 
 				const bool bIsTimerAddedToGraphs =
@@ -3666,6 +3701,8 @@ void STimersView::ContextMenu_ExportTimers_Execute() const
 
 	FTimingExporter Exporter(*Session.Get());
 	FTimingExporter::FExportTimersParams Params; // default
+	Params.bShowFormatInTimerNames = FInsightsManager::Get()->GetSettings().GetShowFormatInTimerNames();
+
 	Exporter.ExportTimersAsText(*Filename, Params);
 }
 
@@ -4187,6 +4224,17 @@ void STimersView::SetSourceFileSearchPaths_Execute()
 
 ////////////////////////////////////////////////////////////////////////////////////////////////////
 
+void STimersView::ToggleShowFormatInTimerNames_Execute()
+{
+	FInsightsSettings& Settings = FInsightsManager::Get()->GetSettings();
+	Settings.SetAndSaveShowFormatInTimerNames(!Settings.GetShowFormatInTimerNames());
+
+	RebuildTree(true);
+	FTimingProfilerManager::Get()->UpdateCallersAndCallees();
+}
+
+////////////////////////////////////////////////////////////////////////////////////////////////////
+
 } // namespace UE::Insights::TimingProfiler
 
 #undef LOCTEXT_NAMESPACE
diff --git a/Engine/Source/Developer/TraceInsights/Private/Insights/TimingProfiler/Widgets/STimersView.h b/Engine/Source/Developer/TraceInsights/Private/Insights/TimingProfiler/Widgets/STimersView.h
--- a/Engine/Source/Developer/TraceInsights/Private/Insights/TimingProfiler/Widgets/STimersView.h
+++ b/Engine/Source/Developer/TraceInsights/Private/Insights/TimingProfiler/Widgets/STimersView.h
@@ -409,6 +409,8 @@ private:
 
 	FText SetSourceFileSearchPaths_GetTooltip() const;
 	void  SetSourceFileSearchPaths_Execute();
+	
+	void ToggleShowFormatInTimerNames_Execute();
 
 	TSharedPtr<STimingView> GetTimingView() const;
 
diff --git a/Engine/Source/Developer/TraceServices/Private/Model/TimingProfiler.cpp b/Engine/Source/Developer/TraceServices/Private/Model/TimingProfiler.cpp
--- a/Engine/Source/Developer/TraceServices/Private/Model/TimingProfiler.cpp
+++ b/Engine/Source/Developer/TraceServices/Private/Model/TimingProfiler.cpp
@@ -428,6 +428,31 @@ const FMetadataSpec* FTimingProfilerProvider::GetMetadataSpec(uint32 MetadataSpe
 	return &MetadataSpecs[MetadataSpecId];
 }
 
+void FTimingProfilerProvider::MakeTimerDisplayName(const FTimingProfilerTimer& InTimer, TStringBuilder<1024>& InOutDisplayName) const
+{
+	Session.ReadAccessCheck();
+
+	if (InTimer.Name)
+	{
+		InOutDisplayName.Append(InTimer.Name);
+	}
+
+	if (InTimer.HasValidMetadataSpecId())
+	{
+		const TraceServices::FMetadataSpec* Spec = GetMetadataSpec(InTimer.MetadataSpecId);
+		if (Spec && Spec->Format)
+		{
+			TStringView Format = Spec->Format;
+			Format.TrimStartInline();
+			if (!Format.StartsWith(TEXT("- ")) && !Format.StartsWith(TEXT("= ")))
+			{
+				InOutDisplayName.Append(TEXT(" - "));
+			}
+			InOutDisplayName.Append(Format);
+		}
+	}
+}
+
 IEditableTimeline<FTimingProfilerEvent>& FTimingProfilerProvider::GetCpuThreadEditableTimeline(uint32 ThreadId)
 {
 	Session.WriteAccessCheck();
diff --git a/Engine/Source/Developer/TraceServices/Private/Model/TimingProfilerPrivate.h b/Engine/Source/Developer/TraceServices/Private/Model/TimingProfilerPrivate.h
--- a/Engine/Source/Developer/TraceServices/Private/Model/TimingProfilerPrivate.h
+++ b/Engine/Source/Developer/TraceServices/Private/Model/TimingProfilerPrivate.h
@@ -48,6 +48,7 @@ public:
 	virtual uint32 GetOriginalTimerIdFromMetadata(uint32 MetadataTimerId) const override; // also in ITimingProfilerTimerReader
 	virtual TArrayView<const uint8> GetMetadata(uint32 MetadataTimerId) const override; // also in ITimingProfilerTimerReader
 	virtual const FMetadataSpec* GetMetadataSpec(uint32 MetadataSpecId) const override;
+	virtual void MakeTimerDisplayName(const FTimingProfilerTimer& InTimer, TStringBuilder<1024>& InOutDisplayName) const override;
 
 	// GPU
 
diff --git a/Engine/Source/Developer/TraceServices/Public/TraceServices/Model/TimingProfiler.h b/Engine/Source/Developer/TraceServices/Public/TraceServices/Model/TimingProfiler.h
--- a/Engine/Source/Developer/TraceServices/Public/TraceServices/Model/TimingProfiler.h
+++ b/Engine/Source/Developer/TraceServices/Public/TraceServices/Model/TimingProfiler.h
@@ -301,6 +301,14 @@ public:
 	 */
 	virtual const FMetadataSpec* GetMetadataSpec(uint32 MetadataSpecId) const { return nullptr; }
 
+	/**
+	* Gets the display name for a timer, combining the static name and the format string if present.
+	*
+	* @param InTimer The timer to create the display name for.
+	* @param InOutDisplayName A TStringBuilder where the display name will be created.
+	*/
+	virtual void MakeTimerDisplayName(const FTimingProfilerTimer& InTimer, TStringBuilder<1024>& InOutDisplayName) const {}
+
 	//////////////////////////////////////////////////
 	// GPU
 
