[Insights]
- TraceServices: Added ITimingProfilerTimerReader& GetTimerReader() to ITimingProfilerProvider.
- TraceInsights: Simplified code regarding access to TimerReader (now acts like a provider; no need anymore for geting the reference/pointer from a lambda/callback; no need for extra lock; locks are needed only when TimerReader is actually used).
- TraceInsights: Normalized code style re access and use of ITimingProfilerProvider.
- TraceInsights/ExecuteCommands: Fixed filtering by timer names (case of multiple timers with same name) for various timing profiler export commands (i.e. commands with -timers= parameter).

#Synced-CL 50159576

diff --git a/Engine/Source/Developer/TraceInsights/Private/Insights/TaskGraphProfiler/TaskGraphProfilerManager.cpp b/Engine/Source/Developer/TraceInsights/Private/Insights/TaskGraphProfiler/TaskGraphProfilerManager.cpp
--- a/Engine/Source/Developer/TraceInsights/Private/Insights/TaskGraphProfiler/TaskGraphProfilerManager.cpp
+++ b/Engine/Source/Developer/TraceInsights/Private/Insights/TaskGraphProfiler/TaskGraphProfilerManager.cpp
@@ -670,32 +670,37 @@ int32 FTaskGraphProfilerManager::GetDepthOfTaskExecution(double TaskStartedTime,
 	}
 
 	TSharedPtr<const TraceServices::IAnalysisSession> Session = FInsightsManager::Get()->GetSession();
-	if (Session.IsValid() && TraceServices::ReadTimingProfilerProvider(*Session.Get()))
+	if (Session.IsValid())
 	{
-		TraceServices::FAnalysisSessionReadScope SessionReadScope(*Session.Get());
-
-		const TraceServices::ITimingProfilerProvider& TimingProfilerProvider = *TraceServices::ReadTimingProfilerProvider(*Session.Get());
-
-		TimingProfilerProvider.ReadTimeline(Track->GetTimelineIndex(),
-			[TaskStartedTime, TaskFinishedTime, &Depth](const TraceServices::ITimingProfilerProvider::Timeline& Timeline)
-			{
-				Timeline.EnumerateEventsDownSampled(TaskStartedTime, TaskFinishedTime, 0, [TaskStartedTime, &Depth](bool IsEnter, double Time, const TraceServices::FTimingProfilerEvent& Event)
-					{
-						if (Time < TaskStartedTime)
-						{
-							check(IsEnter);
-							++Depth;
-							return TraceServices::EEventEnumerate::Continue;
-						}
+		const TraceServices::ITimingProfilerProvider* TimingProfilerProvider = TraceServices::ReadTimingProfilerProvider(*Session.Get());
+		if (TimingProfilerProvider)
+		{
+			TraceServices::FAnalysisSessionReadScope SessionReadScope(*Session.Get());
 
-						if (IsEnter)
+			TimingProfilerProvider->ReadTimeline(Track->GetTimelineIndex(),
+				[TaskStartedTime, TaskFinishedTime, &Depth]
+				(const TraceServices::ITimingProfilerProvider::Timeline& Timeline)
+				{
+					Timeline.EnumerateEventsDownSampled(TaskStartedTime, TaskFinishedTime, 0,
+						[TaskStartedTime, &Depth]
+						(bool IsEnter, double Time, const TraceServices::FTimingProfilerEvent& Event)
 						{
-							++Depth;
-						}
-
-						return TraceServices::EEventEnumerate::Stop;
-					});
-			});
+							if (Time < TaskStartedTime)
+							{
+								check(IsEnter);
+								++Depth;
+								return TraceServices::EEventEnumerate::Continue;
+							}
+
+							if (IsEnter)
+							{
+								++Depth;
+							}
+
+							return TraceServices::EEventEnumerate::Stop;
+						});
+				});
+		}
 	}
 
 	return Depth;
diff --git a/Engine/Source/Developer/TraceInsights/Private/Insights/TaskGraphProfiler/ViewModels/TaskTimingTrack.cpp b/Engine/Source/Developer/TraceInsights/Private/Insights/TaskGraphProfiler/ViewModels/TaskTimingTrack.cpp
--- a/Engine/Source/Developer/TraceInsights/Private/Insights/TaskGraphProfiler/ViewModels/TaskTimingTrack.cpp
+++ b/Engine/Source/Developer/TraceInsights/Private/Insights/TaskGraphProfiler/ViewModels/TaskTimingTrack.cpp
@@ -728,15 +728,13 @@ void FTaskTimingTrack::BuildDrawState(ITimingEventsTrackDrawStateBuilder& Builde
 	TraceServices::FAnalysisSessionReadScope SessionReadScope(*Session.Get());
 
 	const TraceServices::ITasksProvider* TasksProvider = TraceServices::ReadTasksProvider(*Session.Get());
-
-	if (TasksProvider == nullptr)
+	if (!TasksProvider)
 	{
 		return;
 	}
 
 	const TraceServices::FTaskInfo* Task = TasksProvider->TryGetTask(TaskId);
-
-	if (Task == nullptr)
+	if (!Task)
 	{
 		return;
 	}
@@ -859,8 +857,7 @@ void FTaskTimingTrack::OnTimingEventSelected(TSharedPtr<const ITimingEvent> InSe
 	TraceServices::FAnalysisSessionReadScope SessionReadScope(*Session.Get());
 
 	const TraceServices::ITasksProvider* TasksProvider = TraceServices::ReadTasksProvider(*Session.Get());
-
-	if (TasksProvider == nullptr)
+	if (!TasksProvider)
 	{
 		return;
 	}
@@ -869,9 +866,9 @@ void FTaskTimingTrack::OnTimingEventSelected(TSharedPtr<const ITimingEvent> InSe
 	GetEventRelations(ThreadEvent);
 
 	uint32 ThreadId = StaticCastSharedRef<const TimingProfiler::FThreadTimingTrack>(ThreadEvent.GetTrack())->GetThreadId();
-	const TraceServices::FTaskInfo* Task = TasksProvider->TryGetTask(ThreadId, ThreadEvent.GetStartTime());
 
-	if (Task != nullptr)
+	const TraceServices::FTaskInfo* Task = TasksProvider->TryGetTask(ThreadId, ThreadEvent.GetStartTime());
+	if (Task)
 	{
 		TaskId = Task->Id;
 		FTaskGraphProfilerManager::Get()->SelectTaskInTaskTable(TaskId);
@@ -914,15 +911,13 @@ const TSharedPtr<const ITimingEvent> FTaskTimingTrack::GetEvent(float InPosX, fl
 		TraceServices::FAnalysisSessionReadScope SessionReadScope(*Session.Get());
 
 		const TraceServices::ITasksProvider* TasksProvider = TraceServices::ReadTasksProvider(*Session.Get());
-
-		if (TasksProvider == nullptr)
+		if (!TasksProvider)
 		{
 			return TimingEvent;
 		}
 
 		const TraceServices::FTaskInfo* Task = TasksProvider->TryGetTask(TaskId);
-
-		if (Task == nullptr)
+		if (!Task)
 		{
 			return TimingEvent;
 		}
@@ -969,25 +964,29 @@ const TSharedPtr<const ITimingEvent> FTaskTimingTrack::GetEvent(float InPosX, fl
 		}
 		else if (bShowDetailInfoOnTaskTrack)
 		{
-			auto GetEventFromRelations = [&TasksProvider, this, EventTime, Depth, SecondsPerPixel](const TArray<TraceServices::FTaskInfo::FRelationInfo>& Relations, int32 RelationIndex, ETaskEventType EventType) -> TSharedPtr<FTaskTrackEvent>
-			{
-				if (RelationIndex < 0 || Relations.Num() <= RelationIndex)
-				{
-					return nullptr;
-				}
-				const TraceServices::FTaskInfo::FRelationInfo& Relation = Relations[RelationIndex];
+			auto GetEventFromRelations =
+				[&TasksProvider, this, EventTime, Depth, SecondsPerPixel]
+				(const TArray<TraceServices::FTaskInfo::FRelationInfo>& Relations, int32 RelationIndex, ETaskEventType EventType) -> TSharedPtr<FTaskTrackEvent>
 				{
-					const TraceServices::FTaskInfo* TaskInfo = TasksProvider->TryGetTask(Relation.RelativeId);
-					if (TaskInfo && EventTime >= TaskInfo->StartedTimestamp - 2 * SecondsPerPixel && EventTime <= TaskInfo->FinishedTimestamp + 2 * SecondsPerPixel)
+					if (RelationIndex < 0 || Relations.Num() <= RelationIndex)
 					{
-						TSharedPtr<FTaskTrackEvent> NewTimingEvent = MakeShared<FTaskTrackEvent>(SharedThis(this), TaskInfo->StartedTimestamp, TaskInfo->FinishedTimestamp, Depth, EventType);
-						NewTimingEvent->SetTaskId(TaskInfo->Id);
-						return NewTimingEvent;
+						return nullptr;
+					}
+					const TraceServices::FTaskInfo::FRelationInfo& Relation = Relations[RelationIndex];
+					{
+						const TraceServices::FTaskInfo* TaskInfo = TasksProvider->TryGetTask(Relation.RelativeId);
+						if (TaskInfo &&
+							EventTime >= TaskInfo->StartedTimestamp - 2 * SecondsPerPixel &&
+							EventTime <= TaskInfo->FinishedTimestamp + 2 * SecondsPerPixel)
+						{
+							TSharedPtr<FTaskTrackEvent> NewTimingEvent = MakeShared<FTaskTrackEvent>(SharedThis(this), TaskInfo->StartedTimestamp, TaskInfo->FinishedTimestamp, Depth, EventType);
+							NewTimingEvent->SetTaskId(TaskInfo->Id);
+							return NewTimingEvent;
+						}
 					}
-				}
 
-				return nullptr;
-			};
+					return nullptr;
+				};
 
 			TimingEvent = GetEventFromRelations(Task->Prerequisites, Depth - 1, ETaskEventType::PrerequisiteStarted);
 
@@ -1039,15 +1038,13 @@ void FTaskTimingTrack::InitTooltip(FTooltipDrawState& InOutTooltip, const ITimin
 	TraceServices::FAnalysisSessionReadScope SessionReadScope(*Session.Get());
 
 	const TraceServices::ITasksProvider* TasksProvider = TraceServices::ReadTasksProvider(*Session.Get());
-
-	if (TasksProvider == nullptr)
+	if (!TasksProvider)
 	{
 		return;
 	}
 
 	const TraceServices::FTaskInfo* Task = TasksProvider->TryGetTask(TaskTrackEvent.GetTaskId());
-
-	if (Task == nullptr)
+	if (!Task)
 	{
 		return;
 	}
@@ -1107,39 +1104,50 @@ void FTaskTimingTrack::GetEventRelations(const FThreadTrackEvent& InSelectedEven
 	double StartTime = InSelectedEvent.GetStartTime();
 	double EndTime = InSelectedEvent.GetEndTime();
 
+	TSharedRef<const TimingProfiler::FThreadTimingTrack> EventTrack = StaticCastSharedRef<const TimingProfiler::FThreadTimingTrack>(InSelectedEvent.GetTrack());
+	uint32 ThreadId = EventTrack->GetThreadId();
+
+	FTaskGraphProfilerManager::Get()->ShowTaskRelations(&InSelectedEvent, ThreadId);
+
 	TSharedPtr<const TraceServices::IAnalysisSession> Session = FInsightsManager::Get()->GetSession();
-	const TraceServices::ITasksProvider* TasksProvider = TraceServices::ReadTasksProvider(*Session.Get());
-	if (TasksProvider)
+	if (!Session.IsValid())
 	{
-		TSharedPtr<TimingProfiler::STimingView> TimingView = SharedState.GetTimingView();
-		if (!TimingView.IsValid())
-		{
-			return;
-		}
+		return;
+	}
 
-		TSharedRef<const TimingProfiler::FThreadTimingTrack> EventTrack = StaticCastSharedRef<const TimingProfiler::FThreadTimingTrack>(InSelectedEvent.GetTrack());
-		uint32 ThreadId = EventTrack->GetThreadId();
+	const TraceServices::ITasksProvider* TasksProvider = TraceServices::ReadTasksProvider(*Session.Get());
+	if (!TasksProvider)
+	{
+		return;
+	}
 
-		FTaskGraphProfilerManager::Get()->ShowTaskRelations(&InSelectedEvent, ThreadId);
+	const TraceServices::ITimingProfilerProvider* TimingProfilerProvider = TraceServices::ReadTimingProfilerProvider(*Session.Get());
+	if (!TimingProfilerProvider)
+	{
+		return;
+	}
 
-		// if it's an event waiting for tasks completeness, add relations to these tasks
-		const TraceServices::ITimingProfilerProvider& TimingProfilerProvider = *TraceServices::ReadTimingProfilerProvider(*Session.Get());
-		const TraceServices::ITimingProfilerTimerReader* TimerReader;
+	{
 		TraceServices::FAnalysisSessionReadScope SessionReadScope(*Session.Get());
-		TimingProfilerProvider.ReadTimers([&TimerReader](const TraceServices::ITimingProfilerTimerReader& Out) { TimerReader = &Out; });
+
+		const TraceServices::ITimingProfilerTimerReader& TimerReader = TimingProfilerProvider->GetTimerReader();
 
 		const FThreadTrackEvent& ThreadTrackEvent = *static_cast<const FThreadTrackEvent*>(&InSelectedEvent);
-		const TraceServices::FTimingProfilerTimer* Timer = TimerReader->GetTimer(ThreadTrackEvent.GetTimerIndex());
-		check(Timer != nullptr);
+		const TraceServices::FTimingProfilerTimer* Timer = TimerReader.GetTimer(ThreadTrackEvent.GetTimerIndex());
+		if (!Timer)
+		{
+			return;
+		}
 
+		// if it's an event waiting for tasks completeness, add relations to these tasks
 		const TraceServices::FWaitingForTasks* Waiting = TasksProvider->TryGetWaiting(Timer->Name, ThreadId, StartTime);
-		if (Waiting != nullptr)
+		if (Waiting)
 		{
 			int32 NumWaitedTasksToShow = FMath::Min(Waiting->Tasks.Num(), MaxTasksToShow);
 			for (int32 TaskIndex = 0; TaskIndex != NumWaitedTasksToShow; ++TaskIndex)
 			{
 				const TraceServices::FTaskInfo* WaitedTask = TasksProvider->TryGetTask(Waiting->Tasks[TaskIndex]);
-				if (WaitedTask != nullptr)
+				if (WaitedTask)
 				{
 					int32 WaitingTaskExecutionDepth = FTaskGraphProfilerManager::Get()->GetDepthOfTaskExecution(WaitedTask->StartedTimestamp, WaitedTask->FinishedTimestamp, WaitedTask->StartedThreadId);
 
@@ -1154,10 +1162,12 @@ void FTaskTimingTrack::GetEventRelations(const FThreadTrackEvent& InSelectedEven
 		for (int32 TaskIndex = 0; TaskIndex != NumTasksToShow; ++TaskIndex)
 		{
 			const TraceServices::FTaskInfo* Task = TasksProvider->TryGetTask(ParallelForTasks[TaskIndex]);
-			int32 TaskExecutionDepth = FTaskGraphProfilerManager::Get()->GetDepthOfTaskExecution(Task->StartedTimestamp, Task->FinishedTimestamp, Task->StartedThreadId);
-
-			FTaskGraphProfilerManager::Get()->AddRelation(&InSelectedEvent, Task->StartedTimestamp, ThreadId, -1, Task->StartedTimestamp, Task->StartedThreadId, TaskExecutionDepth, ETaskEventType::NestedStarted);
-			FTaskGraphProfilerManager::Get()->AddRelation(&InSelectedEvent, Task->CompletedTimestamp, Task->CompletedThreadId, Task->CompletedTimestamp, ThreadId, ETaskEventType::NestedCompleted);
+			if (Task)
+			{
+				int32 TaskExecutionDepth = FTaskGraphProfilerManager::Get()->GetDepthOfTaskExecution(Task->StartedTimestamp, Task->FinishedTimestamp, Task->StartedThreadId);
+				FTaskGraphProfilerManager::Get()->AddRelation(&InSelectedEvent, Task->StartedTimestamp, ThreadId, -1, Task->StartedTimestamp, Task->StartedThreadId, TaskExecutionDepth, ETaskEventType::NestedStarted);
+				FTaskGraphProfilerManager::Get()->AddRelation(&InSelectedEvent, Task->CompletedTimestamp, Task->CompletedThreadId, Task->CompletedTimestamp, ThreadId, ETaskEventType::NestedCompleted);
+			}
 		}
 	}
 }
diff --git a/Engine/Source/Developer/TraceInsights/Private/Insights/Tests/FunctionalTests/TimingProfilerTests.cpp b/Engine/Source/Developer/TraceInsights/Private/Insights/Tests/FunctionalTests/TimingProfilerTests.cpp
--- a/Engine/Source/Developer/TraceInsights/Private/Insights/Tests/FunctionalTests/TimingProfilerTests.cpp
+++ b/Engine/Source/Developer/TraceInsights/Private/Insights/Tests/FunctionalTests/TimingProfilerTests.cpp
@@ -130,35 +130,38 @@ bool EnumerateEventsToFile::RunTest(const FString& Parameters)
 	uint32 TimelineIndex = (uint32)-1;
 
 	TSharedPtr<const TraceServices::IAnalysisSession> Session = UE::Insights::FInsightsManager::Get()->GetSession();
-	if (Session.IsValid() && TraceServices::ReadTimingProfilerProvider(*Session.Get()))
+	if (Session.IsValid())
 	{
-		TraceServices::FAnalysisSessionReadScope SessionReadScope(*Session.Get());
-
-		const TraceServices::ITimingProfilerProvider& TimingProfilerProvider = *TraceServices::ReadTimingProfilerProvider(*Session.Get());
+		const TraceServices::ITimingProfilerProvider* TimingProfilerProvider = TraceServices::ReadTimingProfilerProvider(*Session.Get());
+		if (TimingProfilerProvider)
+		{
+			TraceServices::FAnalysisSessionReadScope SessionReadScope(*Session.Get());
 
-		TimelineIndex = FTimingProfilerTests::GetTimelineIndex(TEXT("GameThread"));
+			TimelineIndex = FTimingProfilerTests::GetTimelineIndex(TEXT("GameThread"));
 
-		AddErrorIfFalse(TimelineIndex != (uint32)-1, TEXT("Failed to get track named GameThread"));
+			AddErrorIfFalse(TimelineIndex != (uint32)-1, TEXT("Failed to get track named GameThread"));
 
-		const TraceServices::ITimingProfilerTimerReader* TimerReader;
-		TimingProfilerProvider.ReadTimers([&TimerReader](const TraceServices::ITimingProfilerTimerReader& Out) { TimerReader = &Out; });
+			const TraceServices::ITimingProfilerTimerReader& TimerReader = TimingProfilerProvider->GetTimerReader();
 
-		FString FilePath = FPaths::ProjectSavedDir() + TEXT("/EnumerateEventsToFile.txt");
-		TUniquePtr<FArchive> ArchiveWriter = TUniquePtr<FArchive>(IFileManager::Get().CreateFileWriter(*FilePath));
-		TimingProfilerProvider.ReadTimeline(TimelineIndex,
-			[Session, &ArchiveWriter, TimerReader](const TraceServices::ITimingProfilerProvider::Timeline& Timeline)
-			{
-				Timeline.EnumerateEvents(-1.0, Session->GetDurationSeconds() + 1.0,
-					[&ArchiveWriter, TimerReader](double EventStartTime, double EventEndTime, uint32 EventDepth, const TraceServices::FTimingProfilerEvent& Event)
-					{
-						const TraceServices::FTimingProfilerTimer* Timer = TimerReader->GetTimer(Event.TimerIndex);
+			FString FilePath = FPaths::ProjectSavedDir() + TEXT("/EnumerateEventsToFile.txt");
+			TUniquePtr<FArchive> ArchiveWriter = TUniquePtr<FArchive>(IFileManager::Get().CreateFileWriter(*FilePath));
+			TimingProfilerProvider->ReadTimeline(TimelineIndex,
+				[Session, &ArchiveWriter, &TimerReader]
+				(const TraceServices::ITimingProfilerProvider::Timeline& Timeline)
+				{
+					Timeline.EnumerateEvents(-1.0, Session->GetDurationSeconds() + 1.0,
+						[&ArchiveWriter, &TimerReader]
+						(double EventStartTime, double EventEndTime, uint32 EventDepth, const TraceServices::FTimingProfilerEvent& Event)
+						{
+							const TraceServices::FTimingProfilerTimer* Timer = TimerReader.GetTimer(Event.TimerIndex);
 
-						(*ArchiveWriter) << FString::Printf(TEXT("%s %f %f %d\n"), Timer->Name, EventStartTime, EventEndTime, EventDepth).GetCharArray();
-						return TraceServices::EEventEnumerate::Continue;
-					});
-			});
+							(*ArchiveWriter) << FString::Printf(TEXT("%s %f %f %d\n"), Timer->Name, EventStartTime, EventEndTime, EventDepth).GetCharArray();
+							return TraceServices::EEventEnumerate::Continue;
+						});
+				});
 
-		ArchiveWriter->Close();
+			ArchiveWriter->Close();
+		}
 	}
 
 	return !HasAnyErrors();
@@ -171,35 +174,38 @@ bool EnumerateScopesToFile::RunTest(const FString& Parameters)
 	uint32 TimelineIndex = (uint32)-1;
 
 	TSharedPtr<const TraceServices::IAnalysisSession> Session = UE::Insights::FInsightsManager::Get()->GetSession();
-	if (Session.IsValid() && TraceServices::ReadTimingProfilerProvider(*Session.Get()))
+	if (Session.IsValid())
 	{
-		TraceServices::FAnalysisSessionReadScope SessionReadScope(*Session.Get());
-
-		const TraceServices::ITimingProfilerProvider& TimingProfilerProvider = *TraceServices::ReadTimingProfilerProvider(*Session.Get());
+		const TraceServices::ITimingProfilerProvider* TimingProfilerProvider = TraceServices::ReadTimingProfilerProvider(*Session.Get());
+		if (TimingProfilerProvider)
+		{
+			TraceServices::FAnalysisSessionReadScope SessionReadScope(*Session.Get());
 
-		TimelineIndex = FTimingProfilerTests::GetTimelineIndex(TEXT("GameThread"));
+			TimelineIndex = FTimingProfilerTests::GetTimelineIndex(TEXT("GameThread"));
 
-		AddErrorIfFalse(TimelineIndex != (uint32)-1, TEXT("Failed to get track named GameThread"));
+			AddErrorIfFalse(TimelineIndex != (uint32)-1, TEXT("Failed to get track named GameThread"));
 
-		const TraceServices::ITimingProfilerTimerReader* TimerReader;
-		TimingProfilerProvider.ReadTimers([&TimerReader](const TraceServices::ITimingProfilerTimerReader& Out) { TimerReader = &Out; });
+			const TraceServices::ITimingProfilerTimerReader& TimerReader = TimingProfilerProvider->GetTimerReader();
 
-		FString FilePath = FPaths::ProjectSavedDir() + TEXT("/EnumerateScopesToFile.txt");
-		TUniquePtr<FArchive> ArchiveWriter = TUniquePtr<FArchive>(IFileManager::Get().CreateFileWriter(*FilePath));
-		TimingProfilerProvider.ReadTimeline(TimelineIndex,
-			[Session, &ArchiveWriter, TimerReader](const TraceServices::ITimingProfilerProvider::Timeline& Timeline)
-			{
-				Timeline.EnumerateEvents(-1.0, Session->GetDurationSeconds() + 1.0,
-					[&ArchiveWriter, TimerReader](bool bStart, double Time, const TraceServices::FTimingProfilerEvent& Event)
-					{
-						const TraceServices::FTimingProfilerTimer* Timer = TimerReader->GetTimer(Event.TimerIndex);
+			FString FilePath = FPaths::ProjectSavedDir() + TEXT("/EnumerateScopesToFile.txt");
+			TUniquePtr<FArchive> ArchiveWriter = TUniquePtr<FArchive>(IFileManager::Get().CreateFileWriter(*FilePath));
+			TimingProfilerProvider->ReadTimeline(TimelineIndex,
+				[Session, &ArchiveWriter, &TimerReader]
+				(const TraceServices::ITimingProfilerProvider::Timeline& Timeline)
+				{
+					Timeline.EnumerateEvents(-1.0, Session->GetDurationSeconds() + 1.0,
+						[&ArchiveWriter, &TimerReader]
+						(bool bStart, double Time, const TraceServices::FTimingProfilerEvent& Event)
+						{
+							const TraceServices::FTimingProfilerTimer* Timer = TimerReader.GetTimer(Event.TimerIndex);
 
-						(*ArchiveWriter) << FString::Printf(TEXT("%s %d %f\n"), Timer->Name, (int32)bStart, Time).GetCharArray();
-						return TraceServices::EEventEnumerate::Continue;
-					});
-			});
+							(*ArchiveWriter) << FString::Printf(TEXT("%s %d %f\n"), Timer->Name, (int32)bStart, Time).GetCharArray();
+							return TraceServices::EEventEnumerate::Continue;
+						});
+				});
 
-		ArchiveWriter->Close();
+			ArchiveWriter->Close();
+		}
 	}
 
 	return !HasAnyErrors();
@@ -351,39 +357,43 @@ void FTimingProfilerTests::RunEnumerateBenchmark(const FEnumerateTestParams& InP
 	uint32 TimelineIndex = (uint32)-1;
 
 	TSharedPtr<const TraceServices::IAnalysisSession> Session = UE::Insights::FInsightsManager::Get()->GetSession();
-	if (Session.IsValid() && TraceServices::ReadTimingProfilerProvider(*Session.Get()))
+	if (Session.IsValid())
 	{
-		TraceServices::FAnalysisSessionReadScope SessionReadScope(*Session.Get());
-
-		const TraceServices::ITimingProfilerProvider& TimingProfilerProvider = *TraceServices::ReadTimingProfilerProvider(*Session.Get());
+		const TraceServices::ITimingProfilerProvider* TimingProfilerProvider = TraceServices::ReadTimingProfilerProvider(*Session.Get());
+		if (TimingProfilerProvider)
+		{
+			TraceServices::FAnalysisSessionReadScope SessionReadScope(*Session.Get());
 
-		SessionTime = Session->GetDurationSeconds();
-		OutCheckValues.SessionDuration = SessionTime;
+			SessionTime = Session->GetDurationSeconds();
+			OutCheckValues.SessionDuration = SessionTime;
 
-		const double TimeIncrement = SessionTime / static_cast<double>(InParams.NumEnumerations);
+			const double TimeIncrement = SessionTime / static_cast<double>(InParams.NumEnumerations);
 
-		TimelineIndex = GetTimelineIndex(TEXT("GameThread"));
+			TimelineIndex = GetTimelineIndex(TEXT("GameThread"));
 
-		TimingProfilerProvider.ReadTimeline(TimelineIndex,
-			[&OutCheckValues, &InParams, TimeIncrement](const TraceServices::ITimingProfilerProvider::Timeline& Timeline)
-			{
-				double Time = 0.0;
-				for (int32 Index = 0; Index < InParams.NumEnumerations; ++Index)
+			TimingProfilerProvider->ReadTimeline(TimelineIndex,
+				[&OutCheckValues, &InParams, TimeIncrement]
+				(const TraceServices::ITimingProfilerProvider::Timeline& Timeline)
 				{
-					Timeline.EnumerateEvents(Time, Time + InParams.Interval,
-						[&OutCheckValues](double EventStartTime, double EventEndTime, uint32 EventDepth, const TraceServices::FTimingProfilerEvent& Event)
-						{
-							EventEndTime = FMath::Min(EventEndTime, OutCheckValues.SessionDuration);
-							OutCheckValues.TotalEventDuration += EventEndTime - EventStartTime;
-							++OutCheckValues.EventCount;
-							OutCheckValues.SumDepth += EventDepth;
-							OutCheckValues.SumTimerIndex += Event.TimerIndex;
-							return TraceServices::EEventEnumerate::Continue;
-						});
+					double Time = 0.0;
+					for (int32 Index = 0; Index < InParams.NumEnumerations; ++Index)
+					{
+						Timeline.EnumerateEvents(Time, Time + InParams.Interval,
+							[&OutCheckValues]
+							(double EventStartTime, double EventEndTime, uint32 EventDepth, const TraceServices::FTimingProfilerEvent& Event)
+							{
+								EventEndTime = FMath::Min(EventEndTime, OutCheckValues.SessionDuration);
+								OutCheckValues.TotalEventDuration += EventEndTime - EventStartTime;
+								++OutCheckValues.EventCount;
+								OutCheckValues.SumDepth += EventDepth;
+								OutCheckValues.SumTimerIndex += Event.TimerIndex;
+								return TraceServices::EEventEnumerate::Continue;
+							});
 
-					Time += TimeIncrement;
-				}
-			});
+						Time += TimeIncrement;
+					}
+				});
+		}
 	}
 
 	Stopwatch.Stop();
@@ -405,58 +415,64 @@ void FTimingProfilerTests::RunEnumerateAsyncBenchmark(const FEnumerateTestParams
 	uint32 TimelineIndex = (uint32)-1;
 
 	TSharedPtr<const TraceServices::IAnalysisSession> Session = UE::Insights::FInsightsManager::Get()->GetSession();
-	if (Session.IsValid() && TraceServices::ReadTimingProfilerProvider(*Session.Get()))
+	if (Session.IsValid())
 	{
-		TraceServices::FAnalysisSessionReadScope SessionReadScope(*Session.Get());
-
-		const TraceServices::ITimingProfilerProvider& TimingProfilerProvider = *TraceServices::ReadTimingProfilerProvider(*Session.Get());
+		const TraceServices::ITimingProfilerProvider* TimingProfilerProvider = TraceServices::ReadTimingProfilerProvider(*Session.Get());
+		if (TimingProfilerProvider)
+		{
+			TraceServices::FAnalysisSessionReadScope SessionReadScope(*Session.Get());
 
-		SessionTime = Session->GetDurationSeconds();
-		OutCheckValues.SessionDuration = SessionTime;
+			SessionTime = Session->GetDurationSeconds();
+			OutCheckValues.SessionDuration = SessionTime;
 
-		const double TimeIncrement = SessionTime / static_cast<double>(InParams.NumEnumerations);
+			const double TimeIncrement = SessionTime / static_cast<double>(InParams.NumEnumerations);
 
-		TimelineIndex = GetTimelineIndex(TEXT("GameThread"));
+			TimelineIndex = GetTimelineIndex(TEXT("GameThread"));
 
-		TArray<FCheckValues> TaskCheckValues;
+			TArray<FCheckValues> TaskCheckValues;
 
-		TimingProfilerProvider.ReadTimeline(TimelineIndex,
-			[SessionTime, &InParams, TimeIncrement, &TaskCheckValues](const TraceServices::ITimingProfilerProvider::Timeline& Timeline)
-			{
-				double Time = 0.0;
-				for (int32 Index = 0; Index < InParams.NumEnumerations; ++Index)
+			TimingProfilerProvider->ReadTimeline(TimelineIndex,
+				[SessionTime, &InParams, TimeIncrement, &TaskCheckValues]
+				(const TraceServices::ITimingProfilerProvider::Timeline& Timeline)
 				{
-					TraceServices::ITimeline<TraceServices::FTimingProfilerEvent>::EnumerateAsyncParams Params;
-					Params.IntervalStart = Time;
-					Params.IntervalEnd = Time + InParams.Interval;
-					Params.Resolution = 0.0;
-					Params.SortOrder = InParams.SortOrder;
-					Params.SetupCallback = [&TaskCheckValues](uint32 NumTasks)
+					double Time = 0.0;
+					for (int32 Index = 0; Index < InParams.NumEnumerations; ++Index)
 					{
-						TaskCheckValues.AddDefaulted(NumTasks);
-					};
-					Params.EventRangeCallback = [&TaskCheckValues, SessionTime](double EventStartTime, double EventEndTime, uint32 EventDepth, const TraceServices::FTimingProfilerEvent& Event, uint32 TaskIndex)
-					{
-						EventEndTime = FMath::Min(EventEndTime, SessionTime);
-						TaskCheckValues[TaskIndex].TotalEventDuration += EventEndTime - EventStartTime;
-						++TaskCheckValues[TaskIndex].EventCount;
-						TaskCheckValues[TaskIndex].SumDepth += EventDepth;
-						TaskCheckValues[TaskIndex].SumTimerIndex += Event.TimerIndex;
-						return TraceServices::EEventEnumerate::Continue;
-					};
+						TraceServices::ITimeline<TraceServices::FTimingProfilerEvent>::EnumerateAsyncParams Params;
+						Params.IntervalStart = Time;
+						Params.IntervalEnd = Time + InParams.Interval;
+						Params.Resolution = 0.0;
+						Params.SortOrder = InParams.SortOrder;
+						Params.SetupCallback =
+							[&TaskCheckValues](uint32 NumTasks)
+							{
+								TaskCheckValues.AddDefaulted(NumTasks);
+							};
+						Params.EventRangeCallback =
+							[&TaskCheckValues, SessionTime]
+							(double EventStartTime, double EventEndTime, uint32 EventDepth, const TraceServices::FTimingProfilerEvent& Event, uint32 TaskIndex)
+							{
+								EventEndTime = FMath::Min(EventEndTime, SessionTime);
+								TaskCheckValues[TaskIndex].TotalEventDuration += EventEndTime - EventStartTime;
+								++TaskCheckValues[TaskIndex].EventCount;
+								TaskCheckValues[TaskIndex].SumDepth += EventDepth;
+								TaskCheckValues[TaskIndex].SumTimerIndex += Event.TimerIndex;
+								return TraceServices::EEventEnumerate::Continue;
+							};
 
-					Timeline.EnumerateEventsDownSampledAsync(Params);
+						Timeline.EnumerateEventsDownSampledAsync(Params);
 
-					Time += TimeIncrement;
-				}
-			});
+						Time += TimeIncrement;
+					}
+				});
 
-		for (auto& CheckValues : TaskCheckValues)
-		{
-			OutCheckValues.TotalEventDuration += CheckValues.TotalEventDuration;
-			OutCheckValues.EventCount += CheckValues.EventCount;
-			OutCheckValues.SumDepth += CheckValues.SumDepth;
-			OutCheckValues.SumTimerIndex += CheckValues.SumTimerIndex;
+			for (auto& CheckValues : TaskCheckValues)
+			{
+				OutCheckValues.TotalEventDuration += CheckValues.TotalEventDuration;
+				OutCheckValues.EventCount += CheckValues.EventCount;
+				OutCheckValues.SumDepth += CheckValues.SumDepth;
+				OutCheckValues.SumTimerIndex += CheckValues.SumTimerIndex;
+			}
 		}
 	}
 
@@ -479,44 +495,49 @@ void FTimingProfilerTests::RunEnumerateAllTracksBenchmark(const FEnumerateTestPa
 	uint32 TimelineIndex = (uint32)-1;
 
 	TSharedPtr<const TraceServices::IAnalysisSession> Session = UE::Insights::FInsightsManager::Get()->GetSession();
-	if (Session.IsValid() && TraceServices::ReadTimingProfilerProvider(*Session.Get()))
+	if (Session.IsValid())
 	{
-		TraceServices::FAnalysisSessionReadScope SessionReadScope(*Session.Get());
-
-		const TraceServices::ITimingProfilerProvider& TimingProfilerProvider = *TraceServices::ReadTimingProfilerProvider(*Session.Get());
+		const TraceServices::ITimingProfilerProvider* TimingProfilerProvider = TraceServices::ReadTimingProfilerProvider(*Session.Get());
+		if (TimingProfilerProvider)
+		{
+			TraceServices::FAnalysisSessionReadScope SessionReadScope(*Session.Get());
 
-		SessionTime = Session->GetDurationSeconds();
-		OutCheckValues.SessionDuration = SessionTime;
+			SessionTime = Session->GetDurationSeconds();
+			OutCheckValues.SessionDuration = SessionTime;
 
-		const double TimeIncrement = SessionTime / static_cast<double>(InParams.NumEnumerations);
+			const double TimeIncrement = SessionTime / static_cast<double>(InParams.NumEnumerations);
 
-		const TraceServices::IThreadProvider& ThreadProvider = TraceServices::ReadThreadProvider(*Session.Get());
-		ThreadProvider.EnumerateThreads(
-			[&TimelineIndex, &TimingProfilerProvider, &OutCheckValues, &InParams, TimeIncrement](const TraceServices::FThreadInfo& ThreadInfo)
-			{
-				TimingProfilerProvider.GetCpuThreadTimelineIndex(ThreadInfo.Id, TimelineIndex);
+			const TraceServices::IThreadProvider& ThreadProvider = TraceServices::ReadThreadProvider(*Session.Get());
+			ThreadProvider.EnumerateThreads(
+				[&TimelineIndex, TimingProfilerProvider, &OutCheckValues, &InParams, TimeIncrement]
+				(const TraceServices::FThreadInfo& ThreadInfo)
+				{
+					TimingProfilerProvider->GetCpuThreadTimelineIndex(ThreadInfo.Id, TimelineIndex);
 
-				TimingProfilerProvider.ReadTimeline(TimelineIndex,
-					[&OutCheckValues, &InParams, TimeIncrement](const TraceServices::ITimingProfilerProvider::Timeline& Timeline)
-					{
-						double Time = 0.0;
-						for (int32 Index = 0; Index < InParams.NumEnumerations; ++Index)
+					TimingProfilerProvider->ReadTimeline(TimelineIndex,
+						[&OutCheckValues, &InParams, TimeIncrement]
+						(const TraceServices::ITimingProfilerProvider::Timeline& Timeline)
 						{
-							Timeline.EnumerateEvents(Time, Time + InParams.Interval,
-								[&OutCheckValues](double EventStartTime, double EventEndTime, uint32 EventDepth, const TraceServices::FTimingProfilerEvent& Event)
-								{
-									EventEndTime = FMath::Min(EventEndTime, OutCheckValues.SessionDuration);
-									OutCheckValues.TotalEventDuration += EventEndTime - EventStartTime;
-									++OutCheckValues.EventCount;
-									OutCheckValues.SumDepth += EventDepth;
-									OutCheckValues.SumTimerIndex += Event.TimerIndex;
-									return TraceServices::EEventEnumerate::Continue;
-								});
-
-							Time += TimeIncrement;
-						}
-					});
-			});
+							double Time = 0.0;
+							for (int32 Index = 0; Index < InParams.NumEnumerations; ++Index)
+							{
+								Timeline.EnumerateEvents(Time, Time + InParams.Interval,
+									[&OutCheckValues]
+									(double EventStartTime, double EventEndTime, uint32 EventDepth, const TraceServices::FTimingProfilerEvent& Event)
+									{
+										EventEndTime = FMath::Min(EventEndTime, OutCheckValues.SessionDuration);
+										OutCheckValues.TotalEventDuration += EventEndTime - EventStartTime;
+										++OutCheckValues.EventCount;
+										OutCheckValues.SumDepth += EventDepth;
+										OutCheckValues.SumTimerIndex += Event.TimerIndex;
+										return TraceServices::EEventEnumerate::Continue;
+									});
+
+								Time += TimeIncrement;
+							}
+						});
+				});
+		}
 	}
 
 	Stopwatch.Stop();
@@ -538,65 +559,72 @@ void FTimingProfilerTests::RunEnumerateAsyncAllTracksBenchmark(const FEnumerateT
 	uint32 TimelineIndex = (uint32)-1;
 
 	TSharedPtr<const TraceServices::IAnalysisSession> Session = UE::Insights::FInsightsManager::Get()->GetSession();
-	if (Session.IsValid() && TraceServices::ReadTimingProfilerProvider(*Session.Get()))
+	if (Session.IsValid())
 	{
-		TraceServices::FAnalysisSessionReadScope SessionReadScope(*Session.Get());
-
-		const TraceServices::ITimingProfilerProvider& TimingProfilerProvider = *TraceServices::ReadTimingProfilerProvider(*Session.Get());
+		const TraceServices::ITimingProfilerProvider* TimingProfilerProvider = TraceServices::ReadTimingProfilerProvider(*Session.Get());
+		if (TimingProfilerProvider)
+		{
+			TraceServices::FAnalysisSessionReadScope SessionReadScope(*Session.Get());
 
-		SessionTime = Session->GetDurationSeconds();
-		OutCheckValues.SessionDuration = SessionTime;
+			SessionTime = Session->GetDurationSeconds();
+			OutCheckValues.SessionDuration = SessionTime;
 
-		const double TimeIncrement = SessionTime / static_cast<double>(InParams.NumEnumerations);
+			const double TimeIncrement = SessionTime / static_cast<double>(InParams.NumEnumerations);
 
-		TArray<FCheckValues> TaskCheckValues;
+			TArray<FCheckValues> TaskCheckValues;
 
-		const TraceServices::IThreadProvider& ThreadProvider = TraceServices::ReadThreadProvider(*Session.Get());
-		ThreadProvider.EnumerateThreads(
-			[&TimelineIndex, &TimingProfilerProvider, &OutCheckValues, &InParams, TimeIncrement, &TaskCheckValues](const TraceServices::FThreadInfo& ThreadInfo)
-			{
-				TimingProfilerProvider.GetCpuThreadTimelineIndex(ThreadInfo.Id, TimelineIndex);
+			const TraceServices::IThreadProvider& ThreadProvider = TraceServices::ReadThreadProvider(*Session.Get());
+			ThreadProvider.EnumerateThreads(
+				[&TimelineIndex, TimingProfilerProvider, &OutCheckValues, &InParams, TimeIncrement, &TaskCheckValues]
+				(const TraceServices::FThreadInfo& ThreadInfo)
+				{
+					TimingProfilerProvider->GetCpuThreadTimelineIndex(ThreadInfo.Id, TimelineIndex);
 
-				TimingProfilerProvider.ReadTimeline(TimelineIndex,
-					[&OutCheckValues, &InParams, TimeIncrement, &TaskCheckValues](const TraceServices::ITimingProfilerProvider::Timeline& Timeline)
-					{
-						double Time = 0.0;
-						for (int32 Index = 0; Index < InParams.NumEnumerations; ++Index)
+					TimingProfilerProvider->ReadTimeline(TimelineIndex,
+						[&OutCheckValues, &InParams, TimeIncrement, &TaskCheckValues]
+						(const TraceServices::ITimingProfilerProvider::Timeline& Timeline)
 						{
-							TraceServices::ITimeline<TraceServices::FTimingProfilerEvent>::EnumerateAsyncParams Params;
-							Params.IntervalStart = Time;
-							Params.IntervalEnd = Time + InParams.Interval;
-							Params.Resolution = 0.0;
-							Params.SortOrder = InParams.SortOrder;
-							Params.SetupCallback = [&TaskCheckValues](uint32 NumTasks)
+							double Time = 0.0;
+							for (int32 Index = 0; Index < InParams.NumEnumerations; ++Index)
 							{
-								TaskCheckValues.AddDefaulted(NumTasks);
-							};
-							Params.EventRangeCallback = [&OutCheckValues, &TaskCheckValues](double EventStartTime, double EventEndTime, uint32 EventDepth, const TraceServices::FTimingProfilerEvent& Event, uint32 TaskIndex)
-							{
-								EventEndTime = FMath::Min(EventEndTime, OutCheckValues.SessionDuration);
-								TaskCheckValues[TaskIndex].TotalEventDuration += EventEndTime - EventStartTime;
-								++TaskCheckValues[TaskIndex].EventCount;
-								TaskCheckValues[TaskIndex].SumDepth += EventDepth;
-								TaskCheckValues[TaskIndex].SumTimerIndex += Event.TimerIndex;
-								return TraceServices::EEventEnumerate::Continue;
-							};
-
-							Timeline.EnumerateEventsDownSampledAsync(Params);
-
-							for (auto& CheckValues : TaskCheckValues)
-							{
-								OutCheckValues.TotalEventDuration += CheckValues.TotalEventDuration;
-								OutCheckValues.EventCount += CheckValues.EventCount;
-								OutCheckValues.SumDepth += CheckValues.SumDepth;
-								OutCheckValues.SumTimerIndex += CheckValues.SumTimerIndex;
+								TraceServices::ITimeline<TraceServices::FTimingProfilerEvent>::EnumerateAsyncParams Params;
+								Params.IntervalStart = Time;
+								Params.IntervalEnd = Time + InParams.Interval;
+								Params.Resolution = 0.0;
+								Params.SortOrder = InParams.SortOrder;
+								Params.SetupCallback =
+									[&TaskCheckValues](uint32 NumTasks)
+									{
+										TaskCheckValues.AddDefaulted(NumTasks);
+									};
+								Params.EventRangeCallback =
+									[&OutCheckValues, &TaskCheckValues]
+									(double EventStartTime, double EventEndTime, uint32 EventDepth, const TraceServices::FTimingProfilerEvent& Event, uint32 TaskIndex)
+									{
+										EventEndTime = FMath::Min(EventEndTime, OutCheckValues.SessionDuration);
+										TaskCheckValues[TaskIndex].TotalEventDuration += EventEndTime - EventStartTime;
+										++TaskCheckValues[TaskIndex].EventCount;
+										TaskCheckValues[TaskIndex].SumDepth += EventDepth;
+										TaskCheckValues[TaskIndex].SumTimerIndex += Event.TimerIndex;
+										return TraceServices::EEventEnumerate::Continue;
+									};
+
+								Timeline.EnumerateEventsDownSampledAsync(Params);
+
+								for (auto& CheckValues : TaskCheckValues)
+								{
+									OutCheckValues.TotalEventDuration += CheckValues.TotalEventDuration;
+									OutCheckValues.EventCount += CheckValues.EventCount;
+									OutCheckValues.SumDepth += CheckValues.SumDepth;
+									OutCheckValues.SumTimerIndex += CheckValues.SumTimerIndex;
+								}
+
+								TaskCheckValues.Empty();
+								Time += TimeIncrement;
 							}
-
-							TaskCheckValues.Empty();
-							Time += TimeIncrement;
-						}
-					});
-			});
+						});
+				});
+		}
 	}
 
 	Stopwatch.Stop();
@@ -609,18 +637,26 @@ void FTimingProfilerTests::RunEnumerateAsyncAllTracksBenchmark(const FEnumerateT
 
 uint32 FTimingProfilerTests::GetTimelineIndex(const TCHAR* InName)
 {
-	TSharedPtr<const TraceServices::IAnalysisSession> Session = UE::Insights::FInsightsManager::Get()->GetSession();
-	const TraceServices::ITimingProfilerProvider& TimingProfilerProvider = *TraceServices::ReadTimingProfilerProvider(*Session.Get());
-	const TraceServices::IThreadProvider& ThreadProvider = TraceServices::ReadThreadProvider(*Session.Get());
 	uint32 TimelineIndex = (uint32)-1;
-	ThreadProvider.EnumerateThreads(
-		[&TimelineIndex, &TimingProfilerProvider, InName](const TraceServices::FThreadInfo& ThreadInfo)
-		{
-			if (!FCString::Strcmp(ThreadInfo.Name, InName))
+
+	// Assumes this is called under FAnalysisSessionReadScope.
+	TSharedPtr<const TraceServices::IAnalysisSession> Session = UE::Insights::FInsightsManager::Get()->GetSession();
+	check(Session.IsValid());
+
+	const TraceServices::ITimingProfilerProvider* TimingProfilerProvider = TraceServices::ReadTimingProfilerProvider(*Session.Get());
+	if (TimingProfilerProvider)
+	{
+		const TraceServices::IThreadProvider& ThreadProvider = TraceServices::ReadThreadProvider(*Session.Get());
+		ThreadProvider.EnumerateThreads(
+			[&TimelineIndex, TimingProfilerProvider, InName]
+			(const TraceServices::FThreadInfo& ThreadInfo)
 			{
-				TimingProfilerProvider.GetCpuThreadTimelineIndex(ThreadInfo.Id, TimelineIndex);
-			}
-		});
+				if (!FCString::Strcmp(ThreadInfo.Name, InName))
+				{
+					TimingProfilerProvider->GetCpuThreadTimelineIndex(ThreadInfo.Id, TimelineIndex);
+				}
+			});
+	}
 
 	return TimelineIndex;
 }
diff --git a/Engine/Source/Developer/TraceInsights/Private/Insights/TimingProfiler/GraphTracks/TimerGraphSeries.cpp b/Engine/Source/Developer/TraceInsights/Private/Insights/TimingProfiler/GraphTracks/TimerGraphSeries.cpp
--- a/Engine/Source/Developer/TraceInsights/Private/Insights/TimingProfiler/GraphTracks/TimerGraphSeries.cpp
+++ b/Engine/Source/Developer/TraceInsights/Private/Insights/TimingProfiler/GraphTracks/TimerGraphSeries.cpp
@@ -65,89 +65,82 @@ void FTimerGraphSeries::Update(FTimingGraphTrack& GraphTrack, const FTimingTrack
 		CachedCpuSamplingTimelineCount = VisibleCpuSamplingThreads.Num();
 		CachedEvents.Empty();
 
+		uint32 NumTimelinesContainingEvent = 0;
+
 		const TraceServices::ITimingProfilerProvider* TimingProfilerProvider = TraceServices::ReadTimingProfilerProvider(*Session.Get());
-		if (!TimingProfilerProvider)
+		if (TimingProfilerProvider)
 		{
-			return;
-		}
-
-		TraceServices::FAnalysisSessionReadScope SessionReadScope(*Session.Get());
+			TraceServices::FAnalysisSessionReadScope SessionReadScope(*Session.Get());
 
-		const TraceServices::ITimingProfilerTimerReader* TimerReader = nullptr;
-		TimingProfilerProvider->ReadTimers([&TimerReader](const TraceServices::ITimingProfilerTimerReader& Out) { TimerReader = &Out; });
-		if (!TimerReader)
-		{
-			return;
-		}
+			const TraceServices::ITimingProfilerTimerReader& TimerReader = TimingProfilerProvider->GetTimerReader();
 
-		uint32 NumTimelinesContainingEvent = 0;
-
-		auto TimelineCallback =
-			[this, SessionDuration, &GraphTrack, TimerReader, &Viewport, &NumTimelinesContainingEvent]
-			(const TraceServices::ITimingProfilerProvider::Timeline& Timeline)
-			{
-				TArray<TArray<FSimpleTimingEvent>> Events;
-
-				TraceServices::ITimingProfilerTimeline::EnumerateAsyncParams Params;
-				Params.IntervalStart = 0;
-				Params.IntervalEnd = SessionDuration;
-				Params.Resolution = 0.0;
-				Params.SetupCallback =
-					[&Events]
-					(uint32 NumTasks)
-					{
-						Events.AddDefaulted(NumTasks);
-					};
-				Params.EventRangeCallback =
-					[this, &Events, TimerReader]
-					(double StartTime, double EndTime, uint32 Depth, const TraceServices::FTimingProfilerEvent& Event, uint32 TaskIndex)
-					{
-						const TraceServices::FTimingProfilerTimer* Timer = TimerReader->GetTimer(Event.TimerIndex);
-						if (ensure(Timer != nullptr))
+			auto TimelineCallback =
+				[this, SessionDuration, &GraphTrack, &TimerReader, &Viewport, &NumTimelinesContainingEvent]
+				(const TraceServices::ITimingProfilerProvider::Timeline& Timeline)
+				{
+					TArray<TArray<FSimpleTimingEvent>> Events;
+
+					TraceServices::ITimingProfilerTimeline::EnumerateAsyncParams Params;
+					Params.IntervalStart = 0;
+					Params.IntervalEnd = SessionDuration;
+					Params.Resolution = 0.0;
+					Params.SetupCallback =
+						[&Events]
+						(uint32 NumTasks)
+						{
+							Events.AddDefaulted(NumTasks);
+						};
+					Params.EventRangeCallback =
+						[this, &Events, &TimerReader]
+						(double StartTime, double EndTime, uint32 Depth, const TraceServices::FTimingProfilerEvent& Event, uint32 TaskIndex)
 						{
-							if (Timer->Id == TimerId)
+							const TraceServices::FTimingProfilerTimer* Timer = TimerReader.GetTimer(Event.TimerIndex);
+							if (ensure(Timer != nullptr))
 							{
-								const double Duration = EndTime - StartTime;
-								Events[TaskIndex].Add({ StartTime, Duration });
+								if (Timer->Id == TimerId)
+								{
+									const double Duration = EndTime - StartTime;
+									Events[TaskIndex].Add({ StartTime, Duration });
+								}
 							}
-						}
-						return TraceServices::EEventEnumerate::Continue;
-					};
-				Timeline.EnumerateEventsDownSampledAsync(Params);
+							return TraceServices::EEventEnumerate::Continue;
+						};
+					Timeline.EnumerateEventsDownSampledAsync(Params);
 
-				int32 NumOfCachedEvents = CachedEvents.Num();
-				for (auto& Array : Events) //-V1078
-				{
-					for (auto& Event : Array) //-V1078
+					int32 NumOfCachedEvents = CachedEvents.Num();
+					for (auto& Array : Events) //-V1078
 					{
-						CachedEvents.Add(Event);
+						for (auto& Event : Array) //-V1078
+						{
+							CachedEvents.Add(Event);
+						}
+						Array.Empty();
 					}
-					Array.Empty();
-				}
 
-				if (NumOfCachedEvents != CachedEvents.Num())
-				{
-					++NumTimelinesContainingEvent;
-				}
-			};
+					if (NumOfCachedEvents != CachedEvents.Num())
+					{
+						++NumTimelinesContainingEvent;
+					}
+				};
 
-		// CPU & GPU timelines
-		for (uint32 TimelineIndex : VisibleTimelines)
-		{
-			TimingProfilerProvider->ReadTimeline(TimelineIndex, TimelineCallback);
-		}
+			// CPU & GPU timelines
+			for (uint32 TimelineIndex : VisibleTimelines)
+			{
+				TimingProfilerProvider->ReadTimeline(TimelineIndex, TimelineCallback);
+			}
 
-		// CPU Sampling timelines
-		if (VisibleCpuSamplingThreads.Num() > 0)
-		{
-			const TraceServices::IStackSamplesProvider* StackSamplesProvider = TraceServices::ReadStackSamplesProvider(*Session.Get());
-			if (StackSamplesProvider)
+			// CPU Sampling timelines
+			if (VisibleCpuSamplingThreads.Num() > 0)
 			{
-				for (uint32 ThreadId : VisibleCpuSamplingThreads)
+				const TraceServices::IStackSamplesProvider* StackSamplesProvider = TraceServices::ReadStackSamplesProvider(*Session.Get());
+				if (StackSamplesProvider)
 				{
-					TraceServices::FProviderReadScopeLock _(*StackSamplesProvider);
-					const TraceServices::ITimingProfilerProvider::Timeline& Timeline = *StackSamplesProvider->GetTimeline(ThreadId);
-					TimelineCallback(Timeline);
+					for (uint32 ThreadId : VisibleCpuSamplingThreads)
+					{
+						TraceServices::FProviderReadScopeLock _(*StackSamplesProvider);
+						const TraceServices::ITimingProfilerProvider::Timeline& Timeline = *StackSamplesProvider->GetTimeline(ThreadId);
+						TimelineCallback(Timeline);
+					}
 				}
 			}
 		}
diff --git a/Engine/Source/Developer/TraceInsights/Private/Insights/TimingProfiler/Tracks/GpuTimingTrack.cpp b/Engine/Source/Developer/TraceInsights/Private/Insights/TimingProfiler/Tracks/GpuTimingTrack.cpp
--- a/Engine/Source/Developer/TraceInsights/Private/Insights/TimingProfiler/Tracks/GpuTimingTrack.cpp
+++ b/Engine/Source/Developer/TraceInsights/Private/Insights/TimingProfiler/Tracks/GpuTimingTrack.cpp
@@ -341,37 +341,41 @@ void FGpuFencesTimingTrack::PreUpdate(const ITimingTrackUpdateContext& Context)
 	{
 		ClearDirtyFlag();
 
-		using namespace TraceServices;
-		TSharedPtr<const IAnalysisSession> Session = FInsightsManager::Get()->GetSession();
-		if (Session.IsValid() && ReadTimingProfilerProvider(*Session.Get()))
+		TSharedPtr<const TraceServices::IAnalysisSession> Session = FInsightsManager::Get()->GetSession();
+		if (Session.IsValid())
 		{
-			FAnalysisSessionReadScope SessionReadScope(*Session.Get());
-
-			const TraceServices::ITimingProfilerProvider& TimingProfilerProvider = *ReadTimingProfilerProvider(*Session.Get());
+			const TraceServices::ITimingProfilerProvider* TimingProfilerProvider = TraceServices::ReadTimingProfilerProvider(*Session.Get());
+			if (TimingProfilerProvider)
+			{
+				TraceServices::FAnalysisSessionReadScope SessionReadScope(*Session.Get());
 
-			FGpuFencesTrackBuilder Builder(*this, Viewport, Context.GetGeometry().Scale);
+				FGpuFencesTrackBuilder Builder(*this, Viewport, Context.GetGeometry().Scale);
 
-			TimingProfilerProvider.EnumerateGpuFences(GetThreadId(), Viewport.GetStartTime(), Viewport.GetEndTime(), [&Builder](const FGpuFenceWrapper& Fence)
-				{
-					if (Fence.FenceType == EGpuFenceType::SignalFence)
-					{
-						const FGpuSignalFence* SignalFence = Fence.Fence.Get<const FGpuSignalFence*>();
-						TCHAR Name[32];
-						FCString::Sprintf(Name, TEXT("%ld"), SignalFence->Value);
-						Builder.AddFence(SignalFence->Timestamp, EGpuFenceType::SignalFence, Name);
-					}
-					else if (Fence.FenceType == EGpuFenceType::WaitFence)
+				TimingProfilerProvider->EnumerateGpuFences(GetThreadId(), Viewport.GetStartTime(), Viewport.GetEndTime(),
+					[&Builder](const TraceServices::FGpuFenceWrapper& Fence)
 					{
-						const FGpuWaitFence* WaitFence = Fence.Fence.Get<const FGpuWaitFence*>();
-						TCHAR Name[32];
-						FCString::Sprintf(Name, TEXT("%ld"), WaitFence->Value);
-						Builder.AddFence(WaitFence->Timestamp, EGpuFenceType::WaitFence, Name);
-					}
-
-					return TraceServices::EEnumerateResult::Continue;
-				});
-
-			Builder.Flush();
+						using namespace TraceServices;
+
+						if (Fence.FenceType == EGpuFenceType::SignalFence)
+						{
+							const FGpuSignalFence* SignalFence = Fence.Fence.Get<const FGpuSignalFence*>();
+							TCHAR Name[32];
+							FCString::Sprintf(Name, TEXT("%ld"), SignalFence->Value);
+							Builder.AddFence(SignalFence->Timestamp, EGpuFenceType::SignalFence, Name);
+						}
+						else if (Fence.FenceType == EGpuFenceType::WaitFence)
+						{
+							const FGpuWaitFence* WaitFence = Fence.Fence.Get<const FGpuWaitFence*>();
+							TCHAR Name[32];
+							FCString::Sprintf(Name, TEXT("%ld"), WaitFence->Value);
+							Builder.AddFence(WaitFence->Timestamp, EGpuFenceType::WaitFence, Name);
+						}
+
+						return TraceServices::EEnumerateResult::Continue;
+					});
+
+				Builder.Flush();
+			}
 		}
 
 		SetNumLanes(TimeMarkerBoxes.Num() ? 1 : 0);
diff --git a/Engine/Source/Developer/TraceInsights/Private/Insights/TimingProfiler/Tracks/ThreadTimingTrack.cpp b/Engine/Source/Developer/TraceInsights/Private/Insights/TimingProfiler/Tracks/ThreadTimingTrack.cpp
--- a/Engine/Source/Developer/TraceInsights/Private/Insights/TimingProfiler/Tracks/ThreadTimingTrack.cpp
+++ b/Engine/Source/Developer/TraceInsights/Private/Insights/TimingProfiler/Tracks/ThreadTimingTrack.cpp
@@ -1731,12 +1731,8 @@ bool FThreadTimingTrackImpl::ReadTimers(TFunctionRef<void(const TraceServices::I
 		if (TimingProfilerProvider)
 		{
 			TraceServices::FAnalysisSessionReadScope SessionReadScope(*Session.Get());
-			TimingProfilerProvider->ReadTimers(
-				[&Callback, TimingProfilerProvider]
-				(const TraceServices::ITimingProfilerTimerReader& TimerReader)
-				{
-					Callback(*TimingProfilerProvider, TimerReader);
-				});
+			const TraceServices::ITimingProfilerTimerReader& TimerReader = TimingProfilerProvider->GetTimerReader();
+			Callback(*TimingProfilerProvider, TimerReader);
 			return true;
 		}
 	}
diff --git a/Engine/Source/Developer/TraceInsights/Private/Insights/TimingProfiler/ViewModels/FrameStatsHelper.cpp b/Engine/Source/Developer/TraceInsights/Private/Insights/TimingProfiler/ViewModels/FrameStatsHelper.cpp
--- a/Engine/Source/Developer/TraceInsights/Private/Insights/TimingProfiler/ViewModels/FrameStatsHelper.cpp
+++ b/Engine/Source/Developer/TraceInsights/Private/Insights/TimingProfiler/ViewModels/FrameStatsHelper.cpp
@@ -35,12 +35,8 @@ void FFrameStatsHelper::ComputeFrameStatsForTimer(TArray<FFrameStatsCachedEvent>
 					[&FrameStatsEvents, TimerId, SessionDuration, TimingProfilerProvider]
 					(const TraceServices::ITimingProfilerTimeline& Timeline)
 					{
-						TimingProfilerProvider->ReadTimers(
-							[&FrameStatsEvents, TimerId, SessionDuration, &Timeline]
-							(const TraceServices::ITimingProfilerTimerReader& TimerReader)
-							{
-								ProcessTimeline(FrameStatsEvents, TimerId, 0.0, SessionDuration, Timeline, TimerReader);
-							});
+						const TraceServices::ITimingProfilerTimerReader& TimerReader = TimingProfilerProvider->GetTimerReader();
+						ProcessTimeline(FrameStatsEvents, TimerId, 0.0, SessionDuration, Timeline, TimerReader);
 					});
 			}
 		}
@@ -69,12 +65,8 @@ void FFrameStatsHelper::ComputeFrameStatsForTimer(TArray<FFrameStatsCachedEvent>
 				[&FrameStatsEvents, TimerId, SessionDuration, TimingProfilerProvider]
 				(const TraceServices::ITimingProfilerTimeline& Timeline)
 				{
-					TimingProfilerProvider->ReadTimers(
-						[&FrameStatsEvents, TimerId, SessionDuration, &Timeline]
-						(const TraceServices::ITimingProfilerTimerReader& TimerReader)
-						{
-							ProcessTimeline(FrameStatsEvents, TimerId, 0.0, SessionDuration, Timeline, TimerReader);
-						});
+					const TraceServices::ITimingProfilerTimerReader& TimerReader = TimingProfilerProvider->GetTimerReader();
+					ProcessTimeline(FrameStatsEvents, TimerId, 0.0, SessionDuration, Timeline, TimerReader);
 				});
 		}
 	}
@@ -98,12 +90,8 @@ void FFrameStatsHelper::ComputeFrameStatsForTimer(TArray<FFrameStatsCachedEvent>
 			Session->ReadAccessCheck();
 			const double SessionDuration = Session->GetDurationSeconds();
 
-			TimingProfilerProvider->ReadTimers(
-				[&FrameStatsEvents, TimerId, &Timeline, SessionDuration]
-				(const TraceServices::ITimingProfilerTimerReader& TimerReader)
-				{
-					ProcessTimeline(FrameStatsEvents, TimerId, 0.0, SessionDuration, Timeline, TimerReader);
-				});
+			const TraceServices::ITimingProfilerTimerReader& TimerReader = TimingProfilerProvider->GetTimerReader();
+			ProcessTimeline(FrameStatsEvents, TimerId, 0.0, SessionDuration, Timeline, TimerReader);
 		}
 	}
 }
diff --git a/Engine/Source/Developer/TraceInsights/Private/Insights/TimingProfiler/ViewModels/ThreadTimingSharedState.cpp b/Engine/Source/Developer/TraceInsights/Private/Insights/TimingProfiler/ViewModels/ThreadTimingSharedState.cpp
--- a/Engine/Source/Developer/TraceInsights/Private/Insights/TimingProfiler/ViewModels/ThreadTimingSharedState.cpp
+++ b/Engine/Source/Developer/TraceInsights/Private/Insights/TimingProfiler/ViewModels/ThreadTimingSharedState.cpp
@@ -390,7 +390,8 @@ bool FThreadTimingSharedState::AreVerseTracksAvailable() const
 	TSharedPtr<const TraceServices::IAnalysisSession> Session = FInsightsManager::Get()->GetSession();
 	if (Session.IsValid())
 	{
-		if (const TraceServices::ITimingProfilerProvider* TimingProfilerProvider = TraceServices::ReadTimingProfilerProvider(*Session.Get()))
+		const TraceServices::ITimingProfilerProvider* TimingProfilerProvider = TraceServices::ReadTimingProfilerProvider(*Session.Get());
+		if (TimingProfilerProvider)
 		{
 			TraceServices::FAnalysisSessionReadScope SessionReadScope(*Session.Get());
 			return TimingProfilerProvider->HasVerseTiming();
@@ -524,49 +525,51 @@ void FThreadTimingSharedState::Tick(Timing::ITimingViewSession& InSession, const
 			int32 CpuTrackOrder = FTimingTrackOrder::Cpu;
 
 			// Iterate through GPU queues.
-			TimingProfilerProvider->EnumerateGpuQueues([this, bShowAllGpuTracks , &InSession, &bTracksOrderChanged, &GpuTrackOrder](const TraceServices::FGpuQueueInfo& QueueInfo)
-			{
-				// Check if there is an available GPU track for this queue.
-				TSharedPtr<FGpuQueueTimingTrack>* TrackPtrPtr = GpuTracks.Find(QueueInfo.Id);
-				if (TrackPtrPtr == nullptr)
+			TimingProfilerProvider->EnumerateGpuQueues(
+				[this, bShowAllGpuTracks , &InSession, &bTracksOrderChanged, &GpuTrackOrder]
+				(const TraceServices::FGpuQueueInfo& QueueInfo)
 				{
-					// Create new Timing Events track for the GPU Queue.
-					TSharedPtr<FGpuQueueTimingTrack> Track = MakeShared<FGpuQueueTimingTrack>(*this, QueueInfo.GetDisplayName(), QueueInfo.TimelineIndex, QueueInfo.Id);
-					Track->SetOrder(GpuTrackOrder);
-					Track->SetVisibilityFlag(bShowAllGpuTracks);
-					GpuTracks.Add(QueueInfo.Id, Track);
-					InSession.AddScrollableTrack(Track);
-
-					if (AreGpuWorkTracksVisible())
+					// Check if there is an available GPU track for this queue.
+					TSharedPtr<FGpuQueueTimingTrack>* TrackPtrPtr = GpuTracks.Find(QueueInfo.Id);
+					if (TrackPtrPtr == nullptr)
 					{
-						// Create the GPU Work track and attach it to the GPU queue track.
-						const FString WorkTrackName = FString::Printf(TEXT("GPU%u - %s %u - WORK"), QueueInfo.GPU, QueueInfo.Name, QueueInfo.Index);
-						TSharedRef<FGpuQueueWorkTimingTrack> WorkTrack = MakeShared<FGpuQueueWorkTimingTrack>(*this, WorkTrackName, QueueInfo.WorkTimelineIndex, QueueInfo.Id);
-						WorkTrack->SetLocation(Track->GetLocation());
-						WorkTrack->SetParentTrack(Track);
-						Track->AddChildTrack(WorkTrack);
-					}
+						// Create new Timing Events track for the GPU Queue.
+						TSharedPtr<FGpuQueueTimingTrack> Track = MakeShared<FGpuQueueTimingTrack>(*this, QueueInfo.GetDisplayName(), QueueInfo.TimelineIndex, QueueInfo.Id);
+						Track->SetOrder(GpuTrackOrder);
+						Track->SetVisibilityFlag(bShowAllGpuTracks);
+						GpuTracks.Add(QueueInfo.Id, Track);
+						InSession.AddScrollableTrack(Track);
 
-					if (AreGpuFencesTracksVisible())
-					{
-						const FString FencesTrackName = FString::Printf(TEXT("GPU%u - %s %u - Fences"), QueueInfo.GPU, QueueInfo.Name, QueueInfo.Index);
-						TSharedRef<FGpuFencesTimingTrack> FencesTrack = MakeShared<FGpuFencesTimingTrack>(*this, FencesTrackName, QueueInfo.Id);
-						FencesTrack->SetLocation(Track->GetLocation());
-						FencesTrack->SetParentTrack(Track);
-						Track->AddChildTrack(FencesTrack);
+						if (AreGpuWorkTracksVisible())
+						{
+							// Create the GPU Work track and attach it to the GPU queue track.
+							const FString WorkTrackName = FString::Printf(TEXT("GPU%u - %s %u - WORK"), QueueInfo.GPU, QueueInfo.Name, QueueInfo.Index);
+							TSharedRef<FGpuQueueWorkTimingTrack> WorkTrack = MakeShared<FGpuQueueWorkTimingTrack>(*this, WorkTrackName, QueueInfo.WorkTimelineIndex, QueueInfo.Id);
+							WorkTrack->SetLocation(Track->GetLocation());
+							WorkTrack->SetParentTrack(Track);
+							Track->AddChildTrack(WorkTrack);
+						}
+
+						if (AreGpuFencesTracksVisible())
+						{
+							const FString FencesTrackName = FString::Printf(TEXT("GPU%u - %s %u - Fences"), QueueInfo.GPU, QueueInfo.Name, QueueInfo.Index);
+							TSharedRef<FGpuFencesTimingTrack> FencesTrack = MakeShared<FGpuFencesTimingTrack>(*this, FencesTrackName, QueueInfo.Id);
+							FencesTrack->SetLocation(Track->GetLocation());
+							FencesTrack->SetParentTrack(Track);
+							Track->AddChildTrack(FencesTrack);
+						}
 					}
-				}
-				else
-				{
-					TSharedPtr<FGpuQueueTimingTrack> Track = *TrackPtrPtr;
-					if (Track->GetOrder() != GpuTrackOrder)
+					else
 					{
-						Track->SetOrder(GpuTrackOrder);
-						bTracksOrderChanged = true;
+						TSharedPtr<FGpuQueueTimingTrack> Track = *TrackPtrPtr;
+						if (Track->GetOrder() != GpuTrackOrder)
+						{
+							Track->SetOrder(GpuTrackOrder);
+							bTracksOrderChanged = true;
+						}
 					}
-				}
-				GpuTrackOrder += 100;
-			});
+					GpuTrackOrder += 100;
+				});
 
 			// Check if we have the Verse sampling timeline.
 			if (!VerseSamplingTrack.IsValid())
@@ -587,89 +590,90 @@ void FThreadTimingSharedState::Tick(Timing::ITimingViewSession& InSession, const
 			ThreadProvider.EnumerateThreads(
 				[this, &InSession, &bTracksOrderChanged, &CpuTrackOrder, TimingProfilerProvider, LoadTimeProfilerProvider]
 				(const TraceServices::FThreadInfo& ThreadInfo)
-			{
-				// Check if this thread is part of a group?
-				const TCHAR* GroupName = ThreadInfo.GroupName;
-				if (!GroupName || *GroupName == 0)
-				{
-					GroupName = ThreadInfo.Name;
-				}
-				if (!GroupName || *GroupName == 0)
-				{
-					GroupName = TEXT("Other Threads");
-				}
-				if (!ThreadGroups.Contains(GroupName))
-				{
-					// Note: The GroupName pointer should be valid for the duration of the session.
-					ThreadGroups.Add(GroupName, { GroupName, 0, CpuTrackOrder });
-				}
-				else
-				{
-					FThreadGroup& ThreadGroup = ThreadGroups[GroupName];
-					ThreadGroup.Order = CpuTrackOrder;
-				}
-
-				const bool bIsGroupVisible = TimingViewSettings->HasCpuThreadGroup(GroupName) ?
-					TimingViewSettings->IsCpuThreadGroupVisible(GroupName) :
-					TimingViewSettings->IsToggleOptionEnabled(FTimingViewSettings::Option_ShowAllCpuTracks);
-
-				// Check if there is an available Asset Loading track for this thread.
-				bool bIsLoadingThread = false;
-				uint32 LoadingTimelineIndex;
-				if (LoadTimeProfilerProvider && LoadTimeProfilerProvider->GetCpuThreadTimelineIndex(ThreadInfo.Id, LoadingTimelineIndex))
-				{
-					bIsLoadingThread = true;
-				}
-
-				// Check if there is an available CPU track for this thread.
-				uint32 CpuTimelineIndex;
-				if (TimingProfilerProvider->GetCpuThreadTimelineIndex(ThreadInfo.Id, CpuTimelineIndex))
 				{
-					TSharedPtr<FCpuTimingTrack>* TrackPtrPtr = CpuTracks.Find(ThreadInfo.Id);
-					if (TrackPtrPtr == nullptr)
+					// Check if this thread is part of a group?
+					const TCHAR* GroupName = ThreadInfo.GroupName;
+					if (!GroupName || *GroupName == 0)
+					{
+						GroupName = ThreadInfo.Name;
+					}
+					if (!GroupName || *GroupName == 0)
+					{
+						GroupName = TEXT("Other Threads");
+					}
+					if (!ThreadGroups.Contains(GroupName))
+					{
+						// Note: The GroupName pointer should be valid for the duration of the session.
+						ThreadGroups.Add(GroupName, { GroupName, 0, CpuTrackOrder });
+					}
+					else
 					{
-						FString TrackName(ThreadInfo.Name && *ThreadInfo.Name ? ThreadInfo.Name : FString::Printf(TEXT("Thread %u"), ThreadInfo.Id));
-
-						// Create new Timing Events track for the CPU thread.
-						TSharedPtr<FCpuTimingTrack> Track = MakeShared<FCpuTimingTrack>(*this, TrackName, GroupName, CpuTimelineIndex, ThreadInfo.Id);
-						Track->SetOrder(CpuTrackOrder);
-						CpuTracks.Add(ThreadInfo.Id, Track);
-
 						FThreadGroup& ThreadGroup = ThreadGroups[GroupName];
-						ThreadGroup.NumTimelines++;
+						ThreadGroup.Order = CpuTrackOrder;
+					}
 
-						if (bIsLoadingThread &&
-							TimingView &&
-							TimingView->GetName() == FInsightsManagerTabs::LoadingProfilerTabId)
-						{
-							Track->SetVisibilityFlag(true);
-						}
-						else
-						{
-							Track->SetVisibilityFlag(bIsGroupVisible);
-						}
+					const bool bIsGroupVisible = TimingViewSettings->HasCpuThreadGroup(GroupName) ?
+						TimingViewSettings->IsCpuThreadGroupVisible(GroupName) :
+						TimingViewSettings->IsToggleOptionEnabled(FTimingViewSettings::Option_ShowAllCpuTracks);
 
-						InSession.AddScrollableTrack(Track);
+					// Check if there is an available Asset Loading track for this thread.
+					bool bIsLoadingThread = false;
+					uint32 LoadingTimelineIndex;
+					if (LoadTimeProfilerProvider && LoadTimeProfilerProvider->GetCpuThreadTimelineIndex(ThreadInfo.Id, LoadingTimelineIndex))
+					{
+						bIsLoadingThread = true;
 					}
-					else
+
+					// Check if there is an available CPU track for this thread.
+					uint32 CpuTimelineIndex;
+					if (TimingProfilerProvider->GetCpuThreadTimelineIndex(ThreadInfo.Id, CpuTimelineIndex))
 					{
-						TSharedPtr<FCpuTimingTrack> Track = *TrackPtrPtr;
-						if (Track->GetOrder() != CpuTrackOrder)
+						TSharedPtr<FCpuTimingTrack>* TrackPtrPtr = CpuTracks.Find(ThreadInfo.Id);
+						if (TrackPtrPtr == nullptr)
 						{
+							FString TrackName(ThreadInfo.Name && *ThreadInfo.Name ? ThreadInfo.Name : FString::Printf(TEXT("Thread %u"), ThreadInfo.Id));
+
+							// Create new Timing Events track for the CPU thread.
+							TSharedPtr<FCpuTimingTrack> Track = MakeShared<FCpuTimingTrack>(*this, TrackName, GroupName, CpuTimelineIndex, ThreadInfo.Id);
 							Track->SetOrder(CpuTrackOrder);
-							bTracksOrderChanged = true;
+							CpuTracks.Add(ThreadInfo.Id, Track);
+
+							FThreadGroup& ThreadGroup = ThreadGroups[GroupName];
+							ThreadGroup.NumTimelines++;
+
+							if (bIsLoadingThread &&
+								TimingView &&
+								TimingView->GetName() == FInsightsManagerTabs::LoadingProfilerTabId)
+							{
+								Track->SetVisibilityFlag(true);
+							}
+							else
+							{
+								Track->SetVisibilityFlag(bIsGroupVisible);
+							}
+
+							InSession.AddScrollableTrack(Track);
+						}
+						else
+						{
+							TSharedPtr<FCpuTimingTrack> Track = *TrackPtrPtr;
+							if (Track->GetOrder() != CpuTrackOrder)
+							{
+								Track->SetOrder(CpuTrackOrder);
+								bTracksOrderChanged = true;
+							}
 						}
 					}
-				}
 
-				constexpr int32 OrderIncrement = FTimingTrackOrder::GroupRange / 1000; // distribute max 1000 tracks in the order group range
-				static_assert(OrderIncrement >= 1, "Order group range too small");
-				CpuTrackOrder += OrderIncrement;
-			});
+					constexpr int32 OrderIncrement = FTimingTrackOrder::GroupRange / 1000; // distribute max 1000 tracks in the order group range
+					static_assert(OrderIncrement >= 1, "Order group range too small");
+					CpuTrackOrder += OrderIncrement;
+				});
 		}
 	}
 
-	if (const TraceServices::IStackSamplesProvider* StackSamplesProvider = TraceServices::ReadStackSamplesProvider(InAnalysisSession))
+	const TraceServices::IStackSamplesProvider* StackSamplesProvider = TraceServices::ReadStackSamplesProvider(InAnalysisSession);
+	if (StackSamplesProvider)
 	{
 		TraceServices::FProviderReadScopeLock _(*StackSamplesProvider);
 
@@ -1038,18 +1042,19 @@ void FThreadTimingSharedState::ToggleOptionState(const TCHAR* CommandName)
 
 TSharedPtr<const ITimingEvent> FThreadTimingSharedState::FindMaxEventInstance(uint32 TimerId, double StartTime, double EndTime, double MaxValidTime)
 {
-	auto CompareAndAssignEvent = [](TSharedPtr<const ITimingEvent>& TimingEvent, TSharedPtr<const ITimingEvent>& TrackEvent)
-	{
-		if (!TrackEvent.IsValid())
+	auto CompareAndAssignEvent =
+		[](TSharedPtr<const ITimingEvent>& TimingEvent, TSharedPtr<const ITimingEvent>& TrackEvent)
 		{
-			return;
-		}
+			if (!TrackEvent.IsValid())
+			{
+				return;
+			}
 
-		if (!TimingEvent.IsValid() || TrackEvent->GetDuration() > TimingEvent->GetDuration())
-		{
-			TimingEvent = TrackEvent;
-		}
-	};
+			if (!TimingEvent.IsValid() || TrackEvent->GetDuration() > TimingEvent->GetDuration())
+			{
+				TimingEvent = TrackEvent;
+			}
+		};
 
 	TSharedPtr<const ITimingEvent> TimingEvent;
 	TSharedPtr<const ITimingEvent> TrackEvent;
@@ -1103,18 +1108,19 @@ TSharedPtr<const ITimingEvent> FThreadTimingSharedState::FindMaxEventInstance(ui
 
 TSharedPtr<const ITimingEvent> FThreadTimingSharedState::FindMinEventInstance(uint32 TimerId, double StartTime, double EndTime, double MaxValidTime)
 {
-	auto CompareAndAssignEvent = [](TSharedPtr<const ITimingEvent>& TimingEvent, TSharedPtr<const ITimingEvent>& TrackEvent)
-	{
-		if (!TrackEvent.IsValid())
+	auto CompareAndAssignEvent =
+		[](TSharedPtr<const ITimingEvent>& TimingEvent, TSharedPtr<const ITimingEvent>& TrackEvent)
 		{
-			return;
-		}
+			if (!TrackEvent.IsValid())
+			{
+				return;
+			}
 
-		if (!TimingEvent.IsValid() || TrackEvent->GetDuration() < TimingEvent->GetDuration())
-		{
-			TimingEvent = TrackEvent;
-		}
-	};
+			if (!TimingEvent.IsValid() || TrackEvent->GetDuration() < TimingEvent->GetDuration())
+			{
+				TimingEvent = TrackEvent;
+			}
+		};
 
 	TSharedPtr<const ITimingEvent> TimingEvent;
 	TSharedPtr<const ITimingEvent> TrackEvent;
@@ -1191,30 +1197,33 @@ void FThreadTimingSharedState::AddGpuWorkChildTracks()
 	}
 
 	const TraceServices::ITimingProfilerProvider* TimingProfilerProvider = TraceServices::ReadTimingProfilerProvider(*Session);
-
-	TraceServices::FAnalysisSessionReadScope SessionReadScope(*Session);
-
-	TimingProfilerProvider->EnumerateGpuQueues([this](const TraceServices::FGpuQueueInfo& QueueInfo)
+	if (TimingProfilerProvider)
 	{
-		// Check if there is an available GPU track for this queue.
-		TSharedPtr<FGpuQueueTimingTrack>* TrackPtrPtr = GpuTracks.Find(QueueInfo.Id);
-		if (TrackPtrPtr != nullptr)
-		{
-			const TSharedPtr<FGpuQueueTimingTrack>& QueueTrack = *TrackPtrPtr;
-			check(QueueTrack.IsValid());
+		TraceServices::FAnalysisSessionReadScope SessionReadScope(*Session);
 
-			// Check if the Work child track is not already created.
-			if (!QueueTrack->FindChildTrackOfType<FGpuQueueWorkTimingTrack>())
+		TimingProfilerProvider->EnumerateGpuQueues(
+			[this](const TraceServices::FGpuQueueInfo& QueueInfo)
 			{
-				// Create the GPU Work track and attach it to the GPU queue track.
-				const FString WorkTrackName = FString::Printf(TEXT("GPU%u - %s %u - Work"), QueueInfo.GPU, QueueInfo.Name, QueueInfo.Index);
-				TSharedRef<FGpuQueueWorkTimingTrack> WorkTrack = MakeShared<FGpuQueueWorkTimingTrack>(*this, WorkTrackName, QueueInfo.WorkTimelineIndex, QueueInfo.Id);
-				WorkTrack->SetLocation(QueueTrack->GetLocation());
-				QueueTrack->AddChildTrack(WorkTrack, 0);
-				WorkTrack->SetParentTrack(QueueTrack);
-			}
-		}
-	});
+				// Check if there is an available GPU track for this queue.
+				TSharedPtr<FGpuQueueTimingTrack>* TrackPtrPtr = GpuTracks.Find(QueueInfo.Id);
+				if (TrackPtrPtr != nullptr)
+				{
+					const TSharedPtr<FGpuQueueTimingTrack>& QueueTrack = *TrackPtrPtr;
+					check(QueueTrack.IsValid());
+
+					// Check if the Work child track is not already created.
+					if (!QueueTrack->FindChildTrackOfType<FGpuQueueWorkTimingTrack>())
+					{
+						// Create the GPU Work track and attach it to the GPU queue track.
+						const FString WorkTrackName = FString::Printf(TEXT("GPU%u - %s %u - Work"), QueueInfo.GPU, QueueInfo.Name, QueueInfo.Index);
+						TSharedRef<FGpuQueueWorkTimingTrack> WorkTrack = MakeShared<FGpuQueueWorkTimingTrack>(*this, WorkTrackName, QueueInfo.WorkTimelineIndex, QueueInfo.Id);
+						WorkTrack->SetLocation(QueueTrack->GetLocation());
+						QueueTrack->AddChildTrack(WorkTrack, 0);
+						WorkTrack->SetParentTrack(QueueTrack);
+					}
+				}
+			});
+	}
 }
 
 ////////////////////////////////////////////////////////////////////////////////////////////////////
@@ -1259,30 +1268,33 @@ void FThreadTimingSharedState::AddGpuFencesChildTracks()
 	}
 
 	const TraceServices::ITimingProfilerProvider* TimingProfilerProvider = TraceServices::ReadTimingProfilerProvider(*Session);
-
-	TraceServices::FAnalysisSessionReadScope SessionReadScope(*Session);
-
-	TimingProfilerProvider->EnumerateGpuQueues([this](const TraceServices::FGpuQueueInfo& QueueInfo)
+	if (TimingProfilerProvider)
 	{
-		// Check if there is an available GPU track for this queue.
-		TSharedPtr<FGpuQueueTimingTrack>* TrackPtrPtr = GpuTracks.Find(QueueInfo.Id);
-		if (TrackPtrPtr != nullptr)
-		{
-			const TSharedPtr<FGpuQueueTimingTrack>& QueueTrack = *TrackPtrPtr;
-			check(QueueTrack.IsValid());
+		TraceServices::FAnalysisSessionReadScope SessionReadScope(*Session);
 
-			// Check if the Fences child track is not already created.
-			if (!QueueTrack->FindChildTrackOfType<FGpuFencesTimingTrack>())
+		TimingProfilerProvider->EnumerateGpuQueues(
+			[this](const TraceServices::FGpuQueueInfo& QueueInfo)
 			{
-				// Create the GPU Fences track and attach it to the GPU queue track.
-				const FString FencesTrackName = FString::Printf(TEXT("GPU%u - %s %u - Fences"), QueueInfo.GPU, QueueInfo.Name, QueueInfo.Index);
-				TSharedRef<FGpuFencesTimingTrack> FencesTrack = MakeShared<FGpuFencesTimingTrack>(*this, FencesTrackName, QueueInfo.Id);
-				FencesTrack->SetLocation(QueueTrack->GetLocation());
-				QueueTrack->AddChildTrack(FencesTrack, QueueTrack->GetChildTracks().Num());
-				FencesTrack->SetParentTrack(QueueTrack);
-			}
-		}
-	});
+				// Check if there is an available GPU track for this queue.
+				TSharedPtr<FGpuQueueTimingTrack>* TrackPtrPtr = GpuTracks.Find(QueueInfo.Id);
+				if (TrackPtrPtr != nullptr)
+				{
+					const TSharedPtr<FGpuQueueTimingTrack>& QueueTrack = *TrackPtrPtr;
+					check(QueueTrack.IsValid());
+
+					// Check if the Fences child track is not already created.
+					if (!QueueTrack->FindChildTrackOfType<FGpuFencesTimingTrack>())
+					{
+						// Create the GPU Fences track and attach it to the GPU queue track.
+						const FString FencesTrackName = FString::Printf(TEXT("GPU%u - %s %u - Fences"), QueueInfo.GPU, QueueInfo.Name, QueueInfo.Index);
+						TSharedRef<FGpuFencesTimingTrack> FencesTrack = MakeShared<FGpuFencesTimingTrack>(*this, FencesTrackName, QueueInfo.Id);
+						FencesTrack->SetLocation(QueueTrack->GetLocation());
+						QueueTrack->AddChildTrack(FencesTrack, QueueTrack->GetChildTracks().Num());
+						FencesTrack->SetParentTrack(QueueTrack);
+					}
+				}
+			});
+	}
 }
 
 ////////////////////////////////////////////////////////////////////////////////////////////////////
@@ -1309,7 +1321,8 @@ void FThreadTimingSharedState::Command_ShowGpuFencesRelations_Execute()
 	bool bVisibility = TimingViewSettings->IsToggleOptionEnabled(FTimingViewSettings::Option_ShowGpuFencesRelations);
 	if (bVisibility == false && TimingView)
 	{
-		TimingView->EditCurrentRelations().RemoveAll([](TUniquePtr<ITimingEventRelation>& Item)
+		TimingView->EditCurrentRelations().RemoveAll(
+			[](TUniquePtr<ITimingEventRelation>& Item)
 			{
 				return Item->Is<FGpuFenceRelation>();
 			});
@@ -1327,7 +1340,8 @@ void FThreadTimingSharedState::OnTimingEventSelected(TSharedPtr<const ITimingEve
 
 	if (TimingView)
 	{
-		TimingView->EditCurrentRelations().RemoveAll([](TUniquePtr<ITimingEventRelation>& Item)
+		TimingView->EditCurrentRelations().RemoveAll(
+			[](TUniquePtr<ITimingEventRelation>& Item)
 			{
 				return Item->Is<FGpuFenceRelation>();
 			});
@@ -1356,8 +1370,6 @@ void FThreadTimingSharedState::OnTimingEventSelected(TSharedPtr<const ITimingEve
 		return;
 	}
 
-	using namespace TraceServices;
-
 	TSharedPtr<const TraceServices::IAnalysisSession> Session = FInsightsManager::Get()->GetSession();
 	if (!Session.IsValid())
 	{
@@ -1365,74 +1377,82 @@ void FThreadTimingSharedState::OnTimingEventSelected(TSharedPtr<const ITimingEve
 	}
 
 	const TraceServices::ITimingProfilerProvider* TimingProfilerProvider = TraceServices::ReadTimingProfilerProvider(*Session);
-
-	if (TimingProfilerProvider == nullptr)
+	if (!TimingProfilerProvider)
 	{
 		return;
 	}
 
-	auto AddFenceRelation = [this, TimingProfilerProvider](const FGpuSignalFence& SignalFence, const FGpuWaitFence& WaitFence, uint32 WaitFenceQueueId)
-	{
-		uint32 Index = 0;
-		uint32 SignalFenceQueueId = WaitFence.QueueToWaitForId;
-		bool Result = TimingProfilerProvider->GetGpuQueueTimelineIndex(SignalFenceQueueId, Index);
-		if (!Result)
-		{
-			return;
-		}
+	using namespace TraceServices;
 
-		int32 SourceDepth = 0;
-		TimingProfilerProvider->ReadTimeline(Index, [SignalFence, &SourceDepth](const TraceServices::ITimingProfilerProvider::Timeline& Timeline)
+	auto AddFenceRelation =
+		[this, TimingProfilerProvider]
+		(const FGpuSignalFence& SignalFence, const FGpuWaitFence& WaitFence, uint32 WaitFenceQueueId)
 		{
-			SourceDepth = Timeline.GetDepthAt(SignalFence.Timestamp) - 1;
-		});
+			uint32 Index = 0;
+			uint32 SignalFenceQueueId = WaitFence.QueueToWaitForId;
+			bool Result = TimingProfilerProvider->GetGpuQueueTimelineIndex(SignalFenceQueueId, Index);
+			if (!Result)
+			{
+				return;
+			}
 
-		Result = TimingProfilerProvider->GetGpuQueueTimelineIndex(WaitFenceQueueId, Index);
-		if (!Result)
-		{
-			return;
-		}
+			int32 SourceDepth = 0;
+			TimingProfilerProvider->ReadTimeline(Index,
+				[SignalFence, &SourceDepth]
+				(const TraceServices::ITimingProfilerProvider::Timeline& Timeline)
+				{
+					SourceDepth = Timeline.GetDepthAt(SignalFence.Timestamp) - 1;
+				});
 
-		int32 TargetDepth = 0;
-		TimingProfilerProvider->ReadTimeline(Index, [WaitFence, &TargetDepth](const TraceServices::ITimingProfilerProvider::Timeline& Timeline)
-		{
-			TargetDepth = Timeline.GetDepthAt(WaitFence.Timestamp) - 1;
-		});
+			Result = TimingProfilerProvider->GetGpuQueueTimelineIndex(WaitFenceQueueId, Index);
+			if (!Result)
+			{
+				return;
+			}
 
-		TUniquePtr<ITimingEventRelation> RelationBase = MakeUnique<FGpuFenceRelation>(SignalFence.Timestamp, SignalFenceQueueId, WaitFence.Timestamp, WaitFenceQueueId);
-		FGpuFenceRelation& Relation = RelationBase->As<FGpuFenceRelation>();
+			int32 TargetDepth = 0;
+			TimingProfilerProvider->ReadTimeline(Index,
+				[WaitFence, &TargetDepth]
+				(const TraceServices::ITimingProfilerProvider::Timeline& Timeline)
+				{
+					TargetDepth = Timeline.GetDepthAt(WaitFence.Timestamp) - 1;
+				});
 
-		Relation.SetSourceDepth(SourceDepth);
-		Relation.SetTargetDepth(TargetDepth);
+			TUniquePtr<ITimingEventRelation> RelationBase = MakeUnique<FGpuFenceRelation>(SignalFence.Timestamp, SignalFenceQueueId, WaitFence.Timestamp, WaitFenceQueueId);
+			FGpuFenceRelation& Relation = RelationBase->As<FGpuFenceRelation>();
 
-		TSharedPtr<FGpuQueueTimingTrack>* Track = GpuTracks.Find(SignalFenceQueueId);
-		if (Track == nullptr)
-		{
-			return;
-		}
-		Relation.SetSourceTrack(*Track);
+			Relation.SetSourceDepth(SourceDepth);
+			Relation.SetTargetDepth(TargetDepth);
 
-		Track = GpuTracks.Find(WaitFenceQueueId);
-		if (Track == nullptr)
-		{
-			return;
-		}
-		Relation.SetTargetTrack(*Track);
+			TSharedPtr<FGpuQueueTimingTrack>* Track = GpuTracks.Find(SignalFenceQueueId);
+			if (Track == nullptr)
+			{
+				return;
+			}
+			Relation.SetSourceTrack(*Track);
 
-		if (TimingView)
-		{
-			TimingView->AddRelation(RelationBase);
-		}
-	};
+			Track = GpuTracks.Find(WaitFenceQueueId);
+			if (Track == nullptr)
+			{
+				return;
+			}
+			Relation.SetTargetTrack(*Track);
+
+			if (TimingView)
+			{
+				TimingView->AddRelation(RelationBase);
+			}
+		};
 
 	TraceServices::FAnalysisSessionReadScope SessionReadScope(*Session);
 
 	TimingProfilerProvider->EnumerateResolvedGpuFences(QueueId, InSelectedEvent->GetStartTime(), InSelectedEvent->GetEndTime(),
-		[AddFenceRelation, TimingProfilerProvider, QueueId](uint32 SignalFenceQueueId, const FGpuSignalFence& SignalFence, uint32 WaitFenceQueueId, const FGpuWaitFence& WaitFence)
-	{
-		AddFenceRelation(SignalFence, WaitFence, WaitFenceQueueId);
-		return EEnumerateResult::Continue;
-	});
+		[AddFenceRelation, TimingProfilerProvider, QueueId]
+		(uint32 SignalFenceQueueId, const FGpuSignalFence& SignalFence, uint32 WaitFenceQueueId, const FGpuWaitFence& WaitFence)
+		{
+			AddFenceRelation(SignalFence, WaitFence, WaitFenceQueueId);
+			return EEnumerateResult::Continue;
+		});
 }
 
 ////////////////////////////////////////////////////////////////////////////////////////////////////
diff --git a/Engine/Source/Developer/TraceInsights/Private/Insights/TimingProfiler/ViewModels/TimerAggregator.cpp b/Engine/Source/Developer/TraceInsights/Private/Insights/TimingProfiler/ViewModels/TimerAggregator.cpp
--- a/Engine/Source/Developer/TraceInsights/Private/Insights/TimingProfiler/ViewModels/TimerAggregator.cpp
+++ b/Engine/Source/Developer/TraceInsights/Private/Insights/TimingProfiler/ViewModels/TimerAggregator.cpp
@@ -105,7 +105,8 @@ void FTimerAggregationWorker::DoWork(TSharedPtr<TraceServices::FCancellationToke
 	bool bHasCpuSamplingTimelines = false;
 	if (CpuSamplingThreads.Num() > 0)
 	{
-		if (const TraceServices::IStackSamplesProvider* StackSamplesProvider = TraceServices::ReadStackSamplesProvider(*Session.Get()))
+		const TraceServices::IStackSamplesProvider* StackSamplesProvider = TraceServices::ReadStackSamplesProvider(*Session.Get());
+		if (StackSamplesProvider)
 		{
 			TraceServices::FProviderReadScopeLock _(*StackSamplesProvider);
 			for (uint32 SystemThreadId : CpuSamplingThreads)
diff --git a/Engine/Source/Developer/TraceInsights/Private/Insights/TimingProfiler/ViewModels/TimerButterflyAggregator.cpp b/Engine/Source/Developer/TraceInsights/Private/Insights/TimingProfiler/ViewModels/TimerButterflyAggregator.cpp
--- a/Engine/Source/Developer/TraceInsights/Private/Insights/TimingProfiler/ViewModels/TimerButterflyAggregator.cpp
+++ b/Engine/Source/Developer/TraceInsights/Private/Insights/TimingProfiler/ViewModels/TimerButterflyAggregator.cpp
@@ -109,7 +109,8 @@ void FTimerButterflyAggregationWorker::DoWork(TSharedPtr<TraceServices::FCancell
 		bool bHasCpuSamplingTimelines = false;
 		if (CpuSamplingThreads.Num() > 0)
 		{
-			if (const TraceServices::IStackSamplesProvider* StackSamplesProvider = TraceServices::ReadStackSamplesProvider(*Session.Get()))
+			const TraceServices::IStackSamplesProvider* StackSamplesProvider = TraceServices::ReadStackSamplesProvider(*Session.Get());
+			if (StackSamplesProvider)
 			{
 				TraceServices::FProviderReadScopeLock _(*StackSamplesProvider);
 				for (uint32 SystemThreadId : CpuSamplingThreads)
diff --git a/Engine/Source/Developer/TraceInsights/Private/Insights/TimingProfiler/ViewModels/TimerFilters.cpp b/Engine/Source/Developer/TraceInsights/Private/Insights/TimingProfiler/ViewModels/TimerFilters.cpp
--- a/Engine/Source/Developer/TraceInsights/Private/Insights/TimingProfiler/ViewModels/TimerFilters.cpp
+++ b/Engine/Source/Developer/TraceInsights/Private/Insights/TimingProfiler/ViewModels/TimerFilters.cpp
@@ -30,20 +30,24 @@ void FTimerNameFilterState::Update()
 	}
 
 	TimerIds.Empty();
+
 	TSharedPtr<const TraceServices::IAnalysisSession> Session = FInsightsManager::Get()->GetSession();
-	if (Session.IsValid() && TraceServices::ReadTimingProfilerProvider(*Session.Get()))
+	if (!Session.IsValid())
 	{
-		TraceServices::FAnalysisSessionReadScope SessionReadScope(*Session.Get());
+		return;
+	}
 
-		const TraceServices::ITimingProfilerProvider& TimingProfilerProvider = *TraceServices::ReadTimingProfilerProvider(*Session.Get());
+	const TraceServices::ITimingProfilerProvider* TimingProfilerProvider = TraceServices::ReadTimingProfilerProvider(*Session.Get());
+	if (TimingProfilerProvider)
+	{
+		TraceServices::FAnalysisSessionReadScope SessionReadScope(*Session.Get());
 
-		const TraceServices::ITimingProfilerTimerReader* TimerReader;
-		TimingProfilerProvider.ReadTimers([&TimerReader](const TraceServices::ITimingProfilerTimerReader& Out) { TimerReader = &Out; });
+		const TraceServices::ITimingProfilerTimerReader& TimerReader = TimingProfilerProvider->GetTimerReader();
 
-		uint32 TimerCount = TimerReader->GetTimerCount();
+		uint32 TimerCount = TimerReader.GetTimerCount();
 		for (uint32 TimerIndex = 0; TimerIndex < TimerCount; ++TimerIndex)
 		{
-			const TraceServices::FTimingProfilerTimer* Timer = TimerReader->GetTimer(TimerIndex);
+			const TraceServices::FTimingProfilerTimer* Timer = TimerReader.GetTimer(TimerIndex);
 			if (Timer && Timer->Name)
 			{
 				if (SelectedOperator->GetKey() == EFilterOperator::Eq)
@@ -129,19 +133,22 @@ FTimerNameFilter::FTimerNameFilter()
 void FTimerNameFilter::PopulateTimerNameSuggestionList(const FString& Text, TArray<FString>& OutSuggestions)
 {
 	TSharedPtr<const TraceServices::IAnalysisSession> Session = FInsightsManager::Get()->GetSession();
-	if (Session.IsValid() && TraceServices::ReadTimingProfilerProvider(*Session.Get()))
+	if (!Session.IsValid())
 	{
-		TraceServices::FAnalysisSessionReadScope SessionReadScope(*Session.Get());
+		return;
+	}
 
-		const TraceServices::ITimingProfilerProvider& TimingProfilerProvider = *TraceServices::ReadTimingProfilerProvider(*Session.Get());
+	const TraceServices::ITimingProfilerProvider* TimingProfilerProvider = TraceServices::ReadTimingProfilerProvider(*Session.Get());
+	if (TimingProfilerProvider)
+	{
+		TraceServices::FAnalysisSessionReadScope SessionReadScope(*Session.Get());
 
-		const TraceServices::ITimingProfilerTimerReader* TimerReader;
-		TimingProfilerProvider.ReadTimers([&TimerReader](const TraceServices::ITimingProfilerTimerReader& Out) { TimerReader = &Out; });
+		const TraceServices::ITimingProfilerTimerReader& TimerReader = TimingProfilerProvider->GetTimerReader();
 
-		uint32 TimerCount = TimerReader->GetTimerCount();
+		uint32 TimerCount = TimerReader.GetTimerCount();
 		for (uint32 TimerIndex = 0; TimerIndex < TimerCount; ++TimerIndex)
 		{
-			const TraceServices::FTimingProfilerTimer* Timer = TimerReader->GetTimer(TimerIndex);
+			const TraceServices::FTimingProfilerTimer* Timer = TimerReader.GetTimer(TimerIndex);
 			if (Timer && Timer->Name)
 			{
 				if (Text.IsEmpty())
@@ -261,33 +268,21 @@ void FMetadataFilterState::Update()
 	ensure(SelectedDataType.IsValid());
 	switch (SelectedDataType->Type)
 	{
-	case EMetadataFilterDataType::Bool:
-	{
-		ConvertedData.Set<bool>(FCString::ToBool(*Term));
-		break;
-	}
-	case EMetadataFilterDataType::Int:
-	{
-		ConvertedData.Set<int64>(FCString::Atoi64(*Term));
-		break;
-	}
-	case EMetadataFilterDataType::Double:
-	{
-		ConvertedData.Set<double>(FCString::Atod(*Term));
-		break;
-	}
-	}
-
-	TSharedPtr<const TraceServices::IAnalysisSession> Session = FInsightsManager::Get()->GetSession();
-	if (Session.IsValid() && TraceServices::ReadTimingProfilerProvider(*Session.Get()))
-	{
-		TraceServices::FAnalysisSessionReadScope SessionReadScope(*Session.Get());
-
-		const TraceServices::ITimingProfilerProvider& TimingProfilerProvider = *TraceServices::ReadTimingProfilerProvider(*Session.Get());
-
-		// We can cache the timer reader because ApplyFilters needs to be called from inside a session lock. Acquiring the lock
-		// in ApplyFilter is too expensive.
-		TimingProfilerProvider.ReadTimers([this](const TraceServices::ITimingProfilerTimerReader& Out) { this->TimerReader = &Out; });
+		case EMetadataFilterDataType::Bool:
+		{
+			ConvertedData.Set<bool>(FCString::ToBool(*Term));
+			break;
+		}
+		case EMetadataFilterDataType::Int:
+		{
+			ConvertedData.Set<int64>(FCString::Atoi64(*Term));
+			break;
+		}
+		case EMetadataFilterDataType::Double:
+		{
+			ConvertedData.Set<double>(FCString::Atod(*Term));
+			break;
+		}
 	}
 }
 
@@ -303,17 +298,27 @@ bool FMetadataFilterState::ApplyFilter(const FBaseFilterContext& Context) const
 	uint32 TimerIndex = static_cast<uint32>(Context.GetDataAsInt64(Filter->GetKey()));
 
 	TSharedPtr<const TraceServices::IAnalysisSession> Session = FInsightsManager::Get()->GetSession();
-	TraceServices::FAnalysisSessionReadScope SessionReadScope(*Session.Get());
-	TArrayView<const uint8> Metadata = TimerReader->GetMetadata(TimerIndex);
-
-	if (Metadata.Num() > 0)
+	if (Session.IsValid())
 	{
-		if (!bShowAllMetadataEvents)
+		const TraceServices::ITimingProfilerProvider* TimingProfilerProvider = TraceServices::ReadTimingProfilerProvider(*Session.Get());
+		if (TimingProfilerProvider)
 		{
-			return ApplyFilterToMetadata(Metadata);
-		}
+			TraceServices::FAnalysisSessionReadScope SessionReadScope(*Session.Get());
+
+			const TraceServices::ITimingProfilerTimerReader& TimerReader = TimingProfilerProvider->GetTimerReader();
+
+			TArrayView<const uint8> Metadata = TimerReader.GetMetadata(TimerIndex);
+
+			if (Metadata.Num() > 0)
+			{
+				if (!bShowAllMetadataEvents)
+				{
+					return ApplyFilterToMetadata(Metadata);
+				}
 
-		return true;
+				return true;
+			}
+		}
 	}
 
 	return false;
diff --git a/Engine/Source/Developer/TraceInsights/Private/Insights/TimingProfiler/ViewModels/TimerFilters.h b/Engine/Source/Developer/TraceInsights/Private/Insights/TimingProfiler/ViewModels/TimerFilters.h
--- a/Engine/Source/Developer/TraceInsights/Private/Insights/TimingProfiler/ViewModels/TimerFilters.h
+++ b/Engine/Source/Developer/TraceInsights/Private/Insights/TimingProfiler/ViewModels/TimerFilters.h
@@ -156,8 +156,6 @@ private:
 	TArray<TSharedPtr<IFilterOperator>> BoolOperators;
 
 	TSharedPtr<SComboBox<TSharedPtr<IFilterOperator>>> OperatorComboBox;
-
-	const TraceServices::ITimingProfilerTimerReader* TimerReader;
 };
 
 ////////////////////////////////////////////////////////////////////////////////////////////////////
diff --git a/Engine/Source/Developer/TraceInsights/Private/Insights/TimingProfiler/ViewModels/TimerNameFilterValueConverter.cpp b/Engine/Source/Developer/TraceInsights/Private/Insights/TimingProfiler/ViewModels/TimerNameFilterValueConverter.cpp
--- a/Engine/Source/Developer/TraceInsights/Private/Insights/TimingProfiler/ViewModels/TimerNameFilterValueConverter.cpp
+++ b/Engine/Source/Developer/TraceInsights/Private/Insights/TimingProfiler/ViewModels/TimerNameFilterValueConverter.cpp
@@ -18,25 +18,26 @@ namespace UE::Insights::TimingProfiler
 bool FTimerNameFilterValueConverter::Convert(const FString& Input, int64& Output, FText& OutError) const
 {
 	TSharedPtr<const TraceServices::IAnalysisSession> Session = FInsightsManager::Get()->GetSession();
-	if (Session.IsValid() && TraceServices::ReadTimingProfilerProvider(*Session.Get()))
+	if (Session.IsValid())
 	{
-		TraceServices::FAnalysisSessionReadScope SessionReadScope(*Session.Get());
-
-		const TraceServices::ITimingProfilerProvider& TimingProfilerProvider = *TraceServices::ReadTimingProfilerProvider(*Session.Get());
+		const TraceServices::ITimingProfilerProvider* TimingProfilerProvider = TraceServices::ReadTimingProfilerProvider(*Session.Get());
+		if (TimingProfilerProvider)
+		{
+			TraceServices::FAnalysisSessionReadScope SessionReadScope(*Session.Get());
 
-		const TraceServices::ITimingProfilerTimerReader* TimerReader;
-		TimingProfilerProvider.ReadTimers([&TimerReader](const TraceServices::ITimingProfilerTimerReader& Out) { TimerReader = &Out; });
+			const TraceServices::ITimingProfilerTimerReader& TimerReader = TimingProfilerProvider->GetTimerReader();
 
-		uint32 TimerCount = TimerReader->GetTimerCount();
-		for (uint32 TimerIndex = 0; TimerIndex < TimerCount; ++TimerIndex)
-		{
-			const TraceServices::FTimingProfilerTimer* Timer = TimerReader->GetTimer(TimerIndex);
-			if (Timer && Timer->Name)
+			uint32 TimerCount = TimerReader.GetTimerCount();
+			for (uint32 TimerIndex = 0; TimerIndex < TimerCount; ++TimerIndex)
 			{
-				if (FCString::Strcmp(Timer->Name, *Input) == 0)
+				const TraceServices::FTimingProfilerTimer* Timer = TimerReader.GetTimer(TimerIndex);
+				if (Timer && Timer->Name)
 				{
-					Output = Timer->Id;
-					return true;
+					if (FCString::Strcmp(Timer->Name, *Input) == 0)
+					{
+						Output = Timer->Id;
+						return true;
+					}
 				}
 			}
 		}
diff --git a/Engine/Source/Developer/TraceInsights/Private/Insights/TimingProfiler/ViewModels/TimerNode.cpp b/Engine/Source/Developer/TraceInsights/Private/Insights/TimingProfiler/ViewModels/TimerNode.cpp
--- a/Engine/Source/Developer/TraceInsights/Private/Insights/TimingProfiler/ViewModels/TimerNode.cpp
+++ b/Engine/Source/Developer/TraceInsights/Private/Insights/TimingProfiler/ViewModels/TimerNode.cpp
@@ -129,17 +129,16 @@ bool FTimerNode::GetSourceFileAndLine(FString& OutFile, uint32& OutLine) const
 	if (GetTimerId() != FTimerNode::InvalidTimerId)
 	{
 		TSharedPtr<const TraceServices::IAnalysisSession> Session = FInsightsManager::Get()->GetSession();
-		if (Session.IsValid() && TraceServices::ReadTimingProfilerProvider(*Session.Get()))
+		if (Session.IsValid())
 		{
-			const TraceServices::ITimingProfilerProvider& TimingProfilerProvider = *TraceServices::ReadTimingProfilerProvider(*Session.Get());
+			const TraceServices::ITimingProfilerProvider* TimingProfilerProvider = TraceServices::ReadTimingProfilerProvider(*Session.Get());
+			if (TimingProfilerProvider)
+			{
+				TraceServices::FAnalysisSessionReadScope SessionReadScope(*Session.Get());
 
-			TraceServices::FAnalysisSessionReadScope SessionReadScope(*Session.Get());
+				const TraceServices::ITimingProfilerTimerReader& TimerReader = TimingProfilerProvider->GetTimerReader();
 
-			const TraceServices::ITimingProfilerTimerReader* TimerReader = nullptr;
-			TimingProfilerProvider.ReadTimers([&TimerReader](const TraceServices::ITimingProfilerTimerReader& Out) { TimerReader = &Out; });
-			if (TimerReader)
-			{
-				const TraceServices::FTimingProfilerTimer* Timer = TimerReader->GetTimer(GetTimerId());
+				const TraceServices::FTimingProfilerTimer* Timer = TimerReader.GetTimer(GetTimerId());
 				if (Timer && Timer->File)
 				{
 					OutFile = FString(Timer->File);
diff --git a/Engine/Source/Developer/TraceInsights/Private/Insights/TimingProfiler/ViewModels/TimingExporter.cpp b/Engine/Source/Developer/TraceInsights/Private/Insights/TimingProfiler/ViewModels/TimingExporter.cpp
--- a/Engine/Source/Developer/TraceInsights/Private/Insights/TimingProfiler/ViewModels/TimingExporter.cpp
+++ b/Engine/Source/Developer/TraceInsights/Private/Insights/TimingProfiler/ViewModels/TimingExporter.cpp
@@ -152,33 +152,40 @@ int32 FTimingExporter::ExportThreadsAsText(const FString& Filename, FExportThrea
 
 	// Write values.
 	{
-		StringBuilder.Appendf(UTF8TEXT("%u"), FGpuTimingTrack::Gpu1ThreadId);
-		Writer.AppendSeparator();
-		StringBuilder.Append(UTF8TEXT("GPU1"));
-		Writer.AppendSeparator();
-		StringBuilder.Append(UTF8TEXT("GPU"));
-		Writer.AppendLineEnd();
-		++ThreadCount;
-
-		StringBuilder.Appendf(UTF8TEXT("%u"), FGpuTimingTrack::Gpu2ThreadId);
-		Writer.AppendSeparator();
-		StringBuilder.Append(UTF8TEXT("GPU2"));
-		Writer.AppendSeparator();
-		StringBuilder.Append(UTF8TEXT("GPU"));
-		Writer.AppendLineEnd();
-		++ThreadCount;
+		// Old GPU Profiler tracks
+		{
+			StringBuilder.Appendf(UTF8TEXT("%u"), FGpuTimingTrack::Gpu1ThreadId);
+			Writer.AppendSeparator();
+			StringBuilder.Append(UTF8TEXT("GPU1"));
+			Writer.AppendSeparator();
+			StringBuilder.Append(UTF8TEXT("GPU"));
+			Writer.AppendLineEnd();
+			++ThreadCount;
+
+			StringBuilder.Appendf(UTF8TEXT("%u"), FGpuTimingTrack::Gpu2ThreadId);
+			Writer.AppendSeparator();
+			StringBuilder.Append(UTF8TEXT("GPU2"));
+			Writer.AppendSeparator();
+			StringBuilder.Append(UTF8TEXT("GPU"));
+			Writer.AppendLineEnd();
+			++ThreadCount;
+		}
 
-		StringBuilder.Appendf(UTF8TEXT("%u"), FVerseTimingTrack::VerseThreadId);
-		Writer.AppendSeparator();
-		StringBuilder.Append(UTF8TEXT("VerseSampling"));
-		Writer.AppendSeparator();
-		StringBuilder.Append(UTF8TEXT("Verse"));
-		Writer.AppendLineEnd();
-		++ThreadCount;
+		// Verse track
+		{
+			StringBuilder.Appendf(UTF8TEXT("%u"), FVerseTimingTrack::VerseThreadId);
+			Writer.AppendSeparator();
+			StringBuilder.Append(UTF8TEXT("VerseSampling"));
+			Writer.AppendSeparator();
+			StringBuilder.Append(UTF8TEXT("Verse"));
+			Writer.AppendLineEnd();
+			++ThreadCount;
+		}
 
 		// Iterate the Frame tracks.
 		{
 			TraceServices::FAnalysisSessionReadScope SessionReadScope(Session);
+
 			const TraceServices::IFrameProvider& FramesProvider = TraceServices::ReadFrameProvider(Session);
 			for (uint32 FrameType = 0; FrameType < (uint32)ETraceFrameType::TraceFrameType_Count; ++FrameType)
 			{
@@ -192,13 +199,13 @@ int32 FTimingExporter::ExportThreadsAsText(const FString& Filename, FExportThrea
 			}
 		}
 
-		// Iterate the Gpu Queues.
+		// Iterate the GPU Queues (new GPU Profiler).
 		{
-			TraceServices::FAnalysisSessionReadScope SessionReadScope(Session);
-
 			const TraceServices::ITimingProfilerProvider* TimingProfilerProvider = TraceServices::ReadTimingProfilerProvider(Session);
 			if (TimingProfilerProvider)
 			{
+				TraceServices::FAnalysisSessionReadScope SessionReadScope(Session);
+
 				TimingProfilerProvider->EnumerateGpuQueues(
 					[&](const TraceServices::FGpuQueueInfo& QueueInfo)
 					{
@@ -296,54 +303,52 @@ int32 FTimingExporter::ExportTimersAsText(const FString& Filename, FExportTimers
 	{
 		TraceServices::FAnalysisSessionReadScope SessionReadScope(Session);
 
-		TimingProfilerProvider->ReadTimers(
-			[&Writer, &StringBuilder, &TimerCount, &TimingProfilerProvider, &Params](const TraceServices::ITimingProfilerTimerReader& TimerReader)
+		const TraceServices::ITimingProfilerTimerReader& TimerReader = TimingProfilerProvider->GetTimerReader();
+
+		TimerCount = TimerReader.GetTimerCount();
+		for (uint32 TimerIndex = 0; TimerIndex < TimerCount; ++TimerIndex)
+		{
+			const TraceServices::FTimingProfilerTimer& Timer = *(TimerReader.GetTimer(TimerIndex));
+			StringBuilder.Appendf(UTF8TEXT("%u"), Timer.Id);
+			Writer.AppendSeparator();
+			switch (Timer.Type)
 			{
-				TimerCount = TimerReader.GetTimerCount();
-				for (uint32 TimerIndex = 0; TimerIndex < TimerCount; ++TimerIndex)
-				{
-					const TraceServices::FTimingProfilerTimer& Timer = *(TimerReader.GetTimer(TimerIndex));
-					StringBuilder.Appendf(UTF8TEXT("%u"), Timer.Id);
-					Writer.AppendSeparator();
-					switch (Timer.Type)
-					{
-					case TraceServices::ETimingProfilerTimerType::GpuScope:
-						StringBuilder.Append(UTF8TEXT("GPU"));
-						break;
-					case TraceServices::ETimingProfilerTimerType::CpuScope:
-						StringBuilder.Append(UTF8TEXT("CPU"));
-						break;
-					case TraceServices::ETimingProfilerTimerType::CpuSampling:
-						StringBuilder.Append(UTF8TEXT("CpuSampling"));
-						break;
-					case TraceServices::ETimingProfilerTimerType::VerseSampling:
-						StringBuilder.Append(UTF8TEXT("VerseSampling"));
-						break;
-					default:
-						StringBuilder.Append(UTF8TEXT("Unknown"));
-						break;
-					}
-					Writer.AppendSeparator();
-					if (Params.bShowFormatInTimerNames)
-					{
-						TStringBuilder<1024> NameStringBuilder;
-						TimingProfilerProvider->MakeTimerDisplayName(Timer, NameStringBuilder);
-						Writer.AppendString(*NameStringBuilder);
-					}
-					else
-					{
-						Writer.AppendString(Timer.Name);
-					}
-					Writer.AppendSeparator();
-					if (Timer.File)
-					{
-						StringBuilder.Append((const UTF8CHAR*)TCHAR_TO_UTF8(Timer.File));
-					}
-					Writer.AppendSeparator();
-					StringBuilder.Appendf(UTF8TEXT("%u"), Timer.Line);
-					Writer.AppendLineEnd();
-				}
-			});
+			case TraceServices::ETimingProfilerTimerType::GpuScope:
+				StringBuilder.Append(UTF8TEXT("GPU"));
+				break;
+			case TraceServices::ETimingProfilerTimerType::CpuScope:
+				StringBuilder.Append(UTF8TEXT("CPU"));
+				break;
+			case TraceServices::ETimingProfilerTimerType::CpuSampling:
+				StringBuilder.Append(UTF8TEXT("CpuSampling"));
+				break;
+			case TraceServices::ETimingProfilerTimerType::VerseSampling:
+				StringBuilder.Append(UTF8TEXT("VerseSampling"));
+				break;
+			default:
+				StringBuilder.Append(UTF8TEXT("Unknown"));
+				break;
+			}
+			Writer.AppendSeparator();
+			if (Params.bShowFormatInTimerNames)
+			{
+				TStringBuilder<1024> NameStringBuilder;
+				TimingProfilerProvider->MakeTimerDisplayName(Timer, NameStringBuilder);
+				Writer.AppendString(*NameStringBuilder);
+			}
+			else
+			{
+				Writer.AppendString(Timer.Name);
+			}
+			Writer.AppendSeparator();
+			if (Timer.File)
+			{
+				StringBuilder.Append((const UTF8CHAR*)TCHAR_TO_UTF8(Timer.File));
+			}
+			Writer.AppendSeparator();
+			StringBuilder.Appendf(UTF8TEXT("%u"), Timer.Line);
+			Writer.AppendLineEnd();
+		}
 	}
 
 	Writer.Flush();
@@ -447,183 +452,164 @@ void FTimingExporter::ExportTimingEvents_WriteHeader(FExportTimingEventsInternal
 
 int32 FTimingExporter::ExportTimingEvents_WriteEvents(FExportTimingEventsInternalParams& Params) const
 {
-	TMap<uint32, const TCHAR*> Timers; // only used if exporting the TimerName column
-
-	if (Params.Columns.Contains(ExportTimingEvents_TimerNameColumn))
-	{
-		// Iterate the GPU & CPU timers.
-		if (TraceServices::ReadTimingProfilerProvider(Session))
-		{
-			TraceServices::FAnalysisSessionReadScope SessionReadScope(Session);
-
-			const TraceServices::ITimingProfilerProvider& TimingProfilerProvider = *TraceServices::ReadTimingProfilerProvider(Session);
-
-			const TraceServices::ITimingProfilerTimerReader* TimerReader;
-			TimingProfilerProvider.ReadTimers([&TimerReader, &Timers](const TraceServices::ITimingProfilerTimerReader& Out) { TimerReader = &Out; });
-
-			uint32 TimerCount = TimerReader->GetTimerCount();
-			for (uint32 TimerIndex = 0; TimerIndex < TimerCount; ++TimerIndex)
-			{
-				const TraceServices::FTimingProfilerTimer& Timer = *(TimerReader->GetTimer(TimerIndex));
-				Timers.Add(Timer.Id, Timer.Name);
-			}
-		}
-	}
-
 	int32 TimingEventCount = 0;
 
-	if (TraceServices::ReadTimingProfilerProvider(Session))
+	const TraceServices::ITimingProfilerProvider* TimingProfilerProvider = TraceServices::ReadTimingProfilerProvider(Session);
+	if (TimingProfilerProvider)
 	{
 		TraceServices::FAnalysisSessionReadScope SessionReadScope(Session);
 
-		const TraceServices::ITimingProfilerProvider& TimingProfilerProvider = *TraceServices::ReadTimingProfilerProvider(Session);
-
-		const TraceServices::ITimingProfilerTimerReader* TimerReader;
-		TimingProfilerProvider.ReadTimers([&TimerReader, &Timers](const TraceServices::ITimingProfilerTimerReader& Out) { TimerReader = &Out; });
+		const TraceServices::ITimingProfilerTimerReader& TimerReader = TimingProfilerProvider->GetTimerReader();
 
-		auto TimelineEnumerator = [&Params, &TimingEventCount, &Timers, TimerReader](const TraceServices::ITimingProfilerProvider::Timeline& Timeline)
-		{
-			// Iterate timing events.
-			Timeline.EnumerateEvents(Params.UserParams.IntervalStartTime, Params.UserParams.IntervalEndTime,
-				[&Params, &TimingEventCount, &Timers, TimerReader](double EventStartTime, double EventEndTime, uint32 EventDepth, const TraceServices::FTimingProfilerEvent& Event)
-				{
-					if (!Params.UserParams.TimingEventFilter || Params.UserParams.TimingEventFilter(EventStartTime, EventEndTime, EventDepth, Event, TimerReader))
+		auto TimelineEnumerator =
+			[&Params, &TimingEventCount, &TimerReader]
+			(const TraceServices::ITimingProfilerProvider::Timeline& Timeline)
+			{
+				// Iterate timing events.
+				Timeline.EnumerateEvents(Params.UserParams.IntervalStartTime, Params.UserParams.IntervalEndTime,
+					[&Params, &TimingEventCount, &TimerReader]
+					(double EventStartTime, double EventEndTime, uint32 EventDepth, const TraceServices::FTimingProfilerEvent& Event)
 					{
-						FUtf8StringBuilder& StringBuilder = Params.Writer.GetStringBuilder();
-						bool bFirst = true;
-						for (const FName& Column : Params.Columns)
+						if (!Params.UserParams.TimingEventFilter || Params.UserParams.TimingEventFilter(EventStartTime, EventEndTime, EventDepth, Event, &TimerReader))
 						{
-							if (Params.Exporter.ExportTimingEventsColumns.Contains(Column))
+							FUtf8StringBuilder& StringBuilder = Params.Writer.GetStringBuilder();
+							bool bFirst = true;
+							for (const FName& Column : Params.Columns)
 							{
-								if (bFirst)
-								{
-									bFirst = false;
-								}
-								else
-								{
-									Params.Writer.AppendSeparator();
-								}
-
-								if (Column == ExportTimingEvents_ThreadIdColumn)
-								{
-									StringBuilder.Appendf(UTF8TEXT("%u"), Params.ThreadId);
-								}
-								else if (Column == ExportTimingEvents_ThreadNameColumn)
-								{
-									Params.Writer.AppendString(Params.ThreadName);
-								}
-								else if (Column == ExportTimingEvents_TimerIdColumn)
+								if (Params.Exporter.ExportTimingEventsColumns.Contains(Column))
 								{
-									uint32 TimerIndex = Event.TimerIndex;
-									if ((int32) TimerIndex < 0)
+									if (bFirst)
+									{
+										bFirst = false;
+									}
+									else
 									{
-										TimerIndex = TimerReader->GetOriginalTimerIdFromMetadata(TimerIndex);
+										Params.Writer.AppendSeparator();
 									}
 
-									StringBuilder.Appendf(UTF8TEXT("%u"), TimerIndex);
-								}
-								else if (Column == ExportTimingEvents_TimerNameColumn)
-								{
-									uint32 TimerIndex = Event.TimerIndex;
-									if ((int32)TimerIndex < 0)
+									if (Column == ExportTimingEvents_ThreadIdColumn)
+									{
+										StringBuilder.Appendf(UTF8TEXT("%u"), Params.ThreadId);
+									}
+									else if (Column == ExportTimingEvents_ThreadNameColumn)
 									{
-										TimerIndex = TimerReader->GetOriginalTimerIdFromMetadata(TimerIndex);
+										Params.Writer.AppendString(Params.ThreadName);
 									}
+									else if (Column == ExportTimingEvents_TimerIdColumn)
+									{
+										uint32 TimerIndex = Event.TimerIndex;
+										if (int32(TimerIndex) < 0)
+										{
+											TimerIndex = TimerReader.GetOriginalTimerIdFromMetadata(TimerIndex);
+										}
 
-									const TCHAR* TimerName = Timers.FindRef(TimerIndex);
-									Params.Writer.AppendString(TimerName);
-								}
-								else if (Column == ExportTimingEvents_StartTimeColumn)
-								{
-									StringBuilder.Appendf(UTF8TEXT("%.9g"), EventStartTime);
-								}
-								else if (Column == ExportTimingEvents_EndTimeColumn)
-								{
-									StringBuilder.Appendf(UTF8TEXT("%.9g"), EventEndTime);
-								}
-								else if (Column == ExportTimingEvents_DurationColumn)
-								{
-									StringBuilder.Appendf(UTF8TEXT("%.9f"), EventEndTime - EventStartTime);
-								}
-								else if (Column == ExportTimingEvents_DepthColumn)
-								{
-									StringBuilder.Appendf(UTF8TEXT("%u"), EventDepth);
+										StringBuilder.Appendf(UTF8TEXT("%u"), TimerIndex);
+									}
+									else if (Column == ExportTimingEvents_TimerNameColumn)
+									{
+										const TraceServices::FTimingProfilerTimer* Timer = TimerReader.GetTimer(Event.TimerIndex);
+										if (Timer != nullptr)
+										{
+											Params.Writer.AppendString(Timer->Name);
+										}
+										else
+										{
+											StringBuilder.Appendf(UTF8TEXT("[0x%x]"), Event.TimerIndex);
+										}
+									}
+									else if (Column == ExportTimingEvents_StartTimeColumn)
+									{
+										StringBuilder.Appendf(UTF8TEXT("%.9g"), EventStartTime);
+									}
+									else if (Column == ExportTimingEvents_EndTimeColumn)
+									{
+										StringBuilder.Appendf(UTF8TEXT("%.9g"), EventEndTime);
+									}
+									else if (Column == ExportTimingEvents_DurationColumn)
+									{
+										StringBuilder.Appendf(UTF8TEXT("%.9f"), EventEndTime - EventStartTime);
+									}
+									else if (Column == ExportTimingEvents_DepthColumn)
+									{
+										StringBuilder.Appendf(UTF8TEXT("%u"), EventDepth);
+									}
 								}
 							}
+							Params.Writer.AppendLineEnd();
+							++TimingEventCount;
 						}
-						Params.Writer.AppendLineEnd();
-						++TimingEventCount;
-					}
 
-					return TraceServices::EEventEnumerate::Continue;
-				});
-		};
+						return TraceServices::EEventEnumerate::Continue;
+					}
+				); // EnumerateEvents
+			};
 
-		auto FrameEnumerator = [&Params, &TimingEventCount](const TraceServices::FFrame& Frame)
-		{
-			FUtf8StringBuilder& StringBuilder = Params.Writer.GetStringBuilder();
-			bool bFirst = true;
-			for (const FName& Column : Params.Columns)
+		auto FrameEnumerator =
+			[&Params, &TimingEventCount]
+			(const TraceServices::FFrame& Frame)
 			{
-				if (Params.Exporter.ExportTimingEventsColumns.Contains(Column))
+				FUtf8StringBuilder& StringBuilder = Params.Writer.GetStringBuilder();
+				bool bFirst = true;
+				for (const FName& Column : Params.Columns)
 				{
-					if (bFirst)
+					if (Params.Exporter.ExportTimingEventsColumns.Contains(Column))
 					{
-						bFirst = false;
-					}
-					else
-					{
-						Params.Writer.AppendSeparator();
-					}
+						if (bFirst)
+						{
+							bFirst = false;
+						}
+						else
+						{
+							Params.Writer.AppendSeparator();
+						}
 
-					if (Column == ExportTimingEvents_ThreadIdColumn)
-					{
-						StringBuilder.Appendf(UTF8TEXT("%u"), Params.ThreadId);
-					}
-					else if (Column == ExportTimingEvents_ThreadNameColumn)
-					{
-						Params.Writer.AppendString(Params.ThreadName);
-					}
-					else if (Column == ExportTimingEvents_TimerIdColumn)
-					{
-						uint32 TimerIndex = (uint32) Frame.Index;
-						StringBuilder.Appendf(UTF8TEXT("%u"), TimerIndex);
-					}
-					else if (Column == ExportTimingEvents_TimerNameColumn)
-					{
-						Params.Writer.AppendString(*FFrameTimingTrack::GetFrameName((uint32)Frame.FrameType, Frame.Index));
-					}
-					else if (Column == ExportTimingEvents_StartTimeColumn)
-					{
-						StringBuilder.Appendf(UTF8TEXT("%.9g"), Frame.StartTime);
-					}
-					else if (Column == ExportTimingEvents_EndTimeColumn)
-					{
-						StringBuilder.Appendf(UTF8TEXT("%.9g"), Frame.EndTime);
-					}
-					else if (Column == ExportTimingEvents_DurationColumn)
-					{
-						StringBuilder.Appendf(UTF8TEXT("%.9f"), Frame.EndTime - Frame.StartTime);
-					}
-					else if (Column == ExportTimingEvents_DepthColumn)
-					{
-						StringBuilder.Appendf(UTF8TEXT("%u"), 1u);
+						if (Column == ExportTimingEvents_ThreadIdColumn)
+						{
+							StringBuilder.Appendf(UTF8TEXT("%u"), Params.ThreadId);
+						}
+						else if (Column == ExportTimingEvents_ThreadNameColumn)
+						{
+							Params.Writer.AppendString(Params.ThreadName);
+						}
+						else if (Column == ExportTimingEvents_TimerIdColumn)
+						{
+							StringBuilder.Appendf(UTF8TEXT("%u"), (uint32)Frame.Index);
+						}
+						else if (Column == ExportTimingEvents_TimerNameColumn)
+						{
+							Params.Writer.AppendString(*FFrameTimingTrack::GetFrameName((uint32)Frame.FrameType, Frame.Index));
+						}
+						else if (Column == ExportTimingEvents_StartTimeColumn)
+						{
+							StringBuilder.Appendf(UTF8TEXT("%.9g"), Frame.StartTime);
+						}
+						else if (Column == ExportTimingEvents_EndTimeColumn)
+						{
+							StringBuilder.Appendf(UTF8TEXT("%.9g"), Frame.EndTime);
+						}
+						else if (Column == ExportTimingEvents_DurationColumn)
+						{
+							StringBuilder.Appendf(UTF8TEXT("%.9f"), Frame.EndTime - Frame.StartTime);
+						}
+						else if (Column == ExportTimingEvents_DepthColumn)
+						{
+							StringBuilder.Appendf(UTF8TEXT("%u"), 1u);
+						}
 					}
 				}
-			}
-			Params.Writer.AppendLineEnd();
-			++TimingEventCount;
-		};
+				Params.Writer.AppendLineEnd();
+				++TimingEventCount;
+			};
 
-		// Iterate the GPU timelines.
+		// Iterate the GPU timelines (old GPU Profiler).
 		{
 			if (!Params.UserParams.ThreadFilter || Params.UserParams.ThreadFilter(FGpuTimingTrack::Gpu1ThreadId))
 			{
 				Params.ThreadId = FGpuTimingTrack::Gpu1ThreadId;
 				Params.ThreadName = TEXT("GPU1");
 				uint32 GpuTimelineIndex1 = 0;
-				TimingProfilerProvider.GetGpuTimelineIndex(GpuTimelineIndex1);
-				TimingProfilerProvider.ReadTimeline(GpuTimelineIndex1, TimelineEnumerator);
+				TimingProfilerProvider->GetGpuTimelineIndex(GpuTimelineIndex1);
+				TimingProfilerProvider->ReadTimeline(GpuTimelineIndex1, TimelineEnumerator);
 			}
 
 			if (!Params.UserParams.ThreadFilter || Params.UserParams.ThreadFilter(FGpuTimingTrack::Gpu2ThreadId))
@@ -631,24 +617,28 @@ int32 FTimingExporter::ExportTimingEvents_WriteEvents(FExportTimingEventsInterna
 				Params.ThreadId = FGpuTimingTrack::Gpu2ThreadId;
 				Params.ThreadName = TEXT("GPU2");
 				uint32 GpuTimelineIndex2 = 0;
-				TimingProfilerProvider.GetGpu2TimelineIndex(GpuTimelineIndex2);
-				TimingProfilerProvider.ReadTimeline(GpuTimelineIndex2, TimelineEnumerator);
+				TimingProfilerProvider->GetGpu2TimelineIndex(GpuTimelineIndex2);
+				TimingProfilerProvider->ReadTimeline(GpuTimelineIndex2, TimelineEnumerator);
 			}
 		}
 
-		// Iterate the GPU timelines for the GPU Queues.
-		TimingProfilerProvider.EnumerateGpuQueues([&TimingProfilerProvider, &TimelineEnumerator, &Params](const TraceServices::FGpuQueueInfo& QueueInfo)
-			{
-				if (!Params.UserParams.ThreadFilter || Params.UserParams.ThreadFilter(GetNonCollidingId(QueueInfo.Id)))
+		// Iterate the GPU timelines for the GPU Queues (new GPU Profiler).
+		{
+			TimingProfilerProvider->EnumerateGpuQueues(
+				[TimingProfilerProvider, &TimelineEnumerator, &Params]
+				(const TraceServices::FGpuQueueInfo& QueueInfo)
 				{
-					FString QueueName = QueueInfo.GetDisplayName();
-					Params.ThreadId = GetNonCollidingId(QueueInfo.Id);
-					Params.ThreadName = *QueueName;
-					TimingProfilerProvider.ReadTimeline(QueueInfo.TimelineIndex, TimelineEnumerator);
+					if (!Params.UserParams.ThreadFilter || Params.UserParams.ThreadFilter(GetNonCollidingId(QueueInfo.Id)))
+					{
+						FString QueueName = QueueInfo.GetDisplayName();
+						Params.ThreadId = GetNonCollidingId(QueueInfo.Id);
+						Params.ThreadName = *QueueName;
+						TimingProfilerProvider->ReadTimeline(QueueInfo.TimelineIndex, TimelineEnumerator);
 
-					Params.ThreadName = nullptr; // The pointer would no longer be valid outside this scope.
-				}
-			});
+						Params.ThreadName = nullptr; // The pointer would no longer be valid outside this scope.
+					}
+				});
+		}
 
 		// Iterate the Verse timelines.
 		{
@@ -657,8 +647,8 @@ int32 FTimingExporter::ExportTimingEvents_WriteEvents(FExportTimingEventsInterna
 				Params.ThreadId = FVerseTimingTrack::VerseThreadId;
 				Params.ThreadName = TEXT("VerseSampling");
 				uint32 VerseTimelineIndex = 0;
-				TimingProfilerProvider.GetVerseTimelineIndex(VerseTimelineIndex);
-				TimingProfilerProvider.ReadTimeline(VerseTimelineIndex, TimelineEnumerator);
+				TimingProfilerProvider->GetVerseTimelineIndex(VerseTimelineIndex);
+				TimingProfilerProvider->ReadTimeline(VerseTimelineIndex, TimelineEnumerator);
 			}
 		}
 
@@ -679,19 +669,22 @@ int32 FTimingExporter::ExportTimingEvents_WriteEvents(FExportTimingEventsInterna
 		}
 
 		// Iterate the CPU threads and their corresponding timelines.
-		const TraceServices::IThreadProvider& ThreadProvider = TraceServices::ReadThreadProvider(Session);
-		ThreadProvider.EnumerateThreads(
-			[&TimelineEnumerator, &Params, &TimingProfilerProvider](const TraceServices::FThreadInfo& ThreadInfo)
-			{
-				if (!Params.UserParams.ThreadFilter || Params.UserParams.ThreadFilter(ThreadInfo.Id))
+		{
+			const TraceServices::IThreadProvider& ThreadProvider = TraceServices::ReadThreadProvider(Session);
+			ThreadProvider.EnumerateThreads(
+				[&TimelineEnumerator, &Params, TimingProfilerProvider]
+				(const TraceServices::FThreadInfo& ThreadInfo)
 				{
-					Params.ThreadId = ThreadInfo.Id;
-					Params.ThreadName = ThreadInfo.Name;
-					uint32 CpuTimelineIndex = 0;
-					TimingProfilerProvider.GetCpuThreadTimelineIndex(ThreadInfo.Id, CpuTimelineIndex);
-					TimingProfilerProvider.ReadTimeline(CpuTimelineIndex, TimelineEnumerator);
-				}
-			});
+					if (!Params.UserParams.ThreadFilter || Params.UserParams.ThreadFilter(ThreadInfo.Id))
+					{
+						Params.ThreadId = ThreadInfo.Id;
+						Params.ThreadName = ThreadInfo.Name;
+						uint32 CpuTimelineIndex = 0;
+						TimingProfilerProvider->GetCpuThreadTimelineIndex(ThreadInfo.Id, CpuTimelineIndex);
+						TimingProfilerProvider->ReadTimeline(CpuTimelineIndex, TimelineEnumerator);
+					}
+				});
+		}
 	}
 
 	return TimingEventCount;
@@ -719,7 +712,8 @@ int32 FTimingExporter::ExportTimingEventsAsTextByRegions(const FString& Filename
 	FExportTimingEventsParams RegionParams = Params;
 	RegionParams.Region.Reset();
 	int32 ExportedRegionCount = EnumerateRegions(RegionGroups, FilenamePattern,
-		[this, &RegionParams](const FString& Filename, const FString& RegionName, double IntervalStartTime, double IntervalEndTime)
+		[this, &RegionParams]
+		(const FString& Filename, const FString& RegionName, double IntervalStartTime, double IntervalEndTime)
 		{
 			RegionParams.IntervalStartTime = IntervalStartTime;
 			RegionParams.IntervalEndTime = IntervalEndTime;
@@ -825,7 +819,8 @@ void FTimingExporter::GetRegions(const FString& InRegionNamePattern, TMap<FStrin
 		UE_LOG(TraceInsights, Log, TEXT("Looking for regions: '%s'"), *InRegionNamePattern);
 
 		RegionProvider.GetDefaultTimeline().EnumerateRegions(0.0, std::numeric_limits<double>::max(),
-			[&RegionCount, &RegionNameSpec, &OutRegionGroups](const TraceServices::FTimeRegion& InRegion) -> bool
+			[&RegionCount, &RegionNameSpec, &OutRegionGroups]
+			(const TraceServices::FTimeRegion& InRegion) -> bool
 			{
 				if (InRegion.Timer == nullptr || InRegion.Timer->Name == nullptr)
 				{
@@ -1075,7 +1070,8 @@ int32 FTimingExporter::ExportTimerCalleesByRegions(const FString& FilenamePatter
 	FExportTimerCalleesParams RegionParams = Params;
 	RegionParams.Region.Reset();
 	int32 ExportedRegionCount = EnumerateRegions(RegionGroups, FilenamePattern,
-		[this, &RegionParams](const FString& Filename, const FString& RegionName, double IntervalStartTime, double IntervalEndTime)
+		[this, &RegionParams]
+		(const FString& Filename, const FString& RegionName, double IntervalStartTime, double IntervalEndTime)
 		{
 			RegionParams.IntervalStartTime = IntervalStartTime;
 			RegionParams.IntervalEndTime = IntervalEndTime;
@@ -1174,10 +1170,12 @@ int32 FTimingExporter::ExportTimerCalleesAsText(const FString& Filename, const F
 	{
 		TraceServices::FAnalysisSessionReadScope SessionReadScope(Session);
 		const TraceServices::IFrameProvider& FrameProvider = TraceServices::ReadFrameProvider(Session);
-		FrameProvider.EnumerateFrames(ETraceFrameType::TraceFrameType_Game, Params.IntervalStartTime, Params.IntervalEndTime, [&NumFrames](const TraceServices::FFrame& Frame)
-		{
-			++NumFrames;
-		});
+		FrameProvider.EnumerateFrames(ETraceFrameType::TraceFrameType_Game, Params.IntervalStartTime, Params.IntervalEndTime,
+			[&NumFrames]
+			(const TraceServices::FFrame& Frame)
+			{
+				++NumFrames;
+			});
 	}
 
 	// Write rows
@@ -1244,29 +1242,34 @@ FTimingExporter::FThreadFilterFunc FTimingExporter::MakeThreadFilterInclusive(co
 
 	TMap<FString, uint32> Threads;
 
-	// Add the GPU threads for the old GPU Profiler.
-	Threads.Add(FString("GPU"), FGpuTimingTrack::Gpu1ThreadId);
-	Threads.Add(FString("GPU1"), FGpuTimingTrack::Gpu1ThreadId);
-	Threads.Add(FString("GPU2"), FGpuTimingTrack::Gpu2ThreadId);
-
-	// Add Gpu Queues for the new Gpu Profiler
+	// Add the GPU pseudo-threads (old GPU Profiler).
 	{
-		TraceServices::FAnalysisSessionReadScope SessionReadScope(Session);
+		Threads.Add(FString("GPU"), FGpuTimingTrack::Gpu1ThreadId);
+		Threads.Add(FString("GPU1"), FGpuTimingTrack::Gpu1ThreadId);
+		Threads.Add(FString("GPU2"), FGpuTimingTrack::Gpu2ThreadId);
+	}
 
+	// Add the GPU Queue pseudo-threads (new GPU Profiler).
+	{
 		const TraceServices::ITimingProfilerProvider* TimingProfilerProvider = TraceServices::ReadTimingProfilerProvider(Session);
 		if (TimingProfilerProvider)
 		{
-			TimingProfilerProvider->EnumerateGpuQueues([&Threads](const TraceServices::FGpuQueueInfo& QueueInfo)
-			{
-				Threads.Add(QueueInfo.GetDisplayName(), GetNonCollidingId(QueueInfo.Id));
-			});
+			TraceServices::FAnalysisSessionReadScope SessionReadScope(Session);
+
+			TimingProfilerProvider->EnumerateGpuQueues(
+				[&Threads](const TraceServices::FGpuQueueInfo& QueueInfo)
+				{
+					Threads.Add(QueueInfo.GetDisplayName(), GetNonCollidingId(QueueInfo.Id));
+				});
 		}
 	}
 
-	// Add the Verse Sampling fake thread.
-	Threads.Add(FString("VerseSampling"), FVerseTimingTrack::VerseThreadId);
+	// Add the Verse Sampling pseudo-thread.
+	{
+		Threads.Add(FString("VerseSampling"), FVerseTimingTrack::VerseThreadId);
+	}
 
-	// Iterate the CPU threads.
+	// Add the CPU threads.
 	{
 		TraceServices::FAnalysisSessionReadScope SessionReadScope(Session);
 
@@ -1279,10 +1282,12 @@ FTimingExporter::FThreadFilterFunc FTimingExporter::MakeThreadFilterInclusive(co
 			});
 	}
 
-	// Add the frame tracks
+	// Add the Frame pseudo-threads.
 	{
 		TraceServices::FAnalysisSessionReadScope SessionReadScope(Session);
+
 		const TraceServices::IFrameProvider& FramesProvider = TraceServices::ReadFrameProvider(Session);
+
 		for (uint32 FrameType = 0; FrameType < (uint32)ETraceFrameType::TraceFrameType_Count; ++FrameType)
 		{
 			Threads.Add(FFrameTimingTrack::GetFrameTrackName(FrameType), GetFrameTypeNonCollidingId(FrameType));
@@ -1376,46 +1381,29 @@ FTimingExporter::FTimingEventFilterFunc FTimingExporter::MakeTimingEventFilterBy
 
 	OutIncludedTimers.Reset();
 
-	TMultiMap<FString, uint32> Timers;
+	TArray<FString> NamePatterns;
+	InFilterString.ParseIntoArray(NamePatterns, TEXT(","), true);
 
-	// Iterate the GPU & CPU timers.
+	// Iterate the available timers.
 	const TraceServices::ITimingProfilerProvider* TimingProfilerProvider = TraceServices::ReadTimingProfilerProvider(Session);
 	if (TimingProfilerProvider)
 	{
 		TraceServices::FAnalysisSessionReadScope SessionReadScope(Session);
 
-		TimingProfilerProvider->ReadTimers(
-			[&Timers](const TraceServices::ITimingProfilerTimerReader& TimerReader)
-			{
-				uint32 TimerCount = TimerReader.GetTimerCount();
-				for (uint32 TimerIndex = 0; TimerIndex < TimerCount; ++TimerIndex)
-				{
-					const TraceServices::FTimingProfilerTimer& Timer = *(TimerReader.GetTimer(TimerIndex));
-					Timers.Add(FString(Timer.Name), Timer.Id);
-				}
-			});
-	}
-
-	TArray<FString> Filter;
-	InFilterString.ParseIntoArray(Filter, TEXT(","), true);
+		const TraceServices::ITimingProfilerTimerReader& TimerReader = TimingProfilerProvider->GetTimerReader();
 
-	TArray<uint32> Ids;
-	for (const FString& TimerWildcard : Filter)
-	{
-		Ids.Reset();
-
-		Timers.MultiFind(TimerWildcard, Ids, false);
-		if (!Ids.IsEmpty())
-		{
-			OutIncludedTimers.Append(Ids);
-		}
-		else
+		const uint32 TimerCount = TimerReader.GetTimerCount();
+		for (uint32 TimerIndex = 0; TimerIndex < TimerCount; ++TimerIndex)
 		{
-			for (const auto& KeyValuePair : Timers)
+			const TraceServices::FTimingProfilerTimer* Timer = TimerReader.GetTimer(TimerIndex);
+			check(Timer != nullptr);
+
+			FString TimerName(Timer->Name);
+			for (const FString& Pattern : NamePatterns)
 			{
-				if (KeyValuePair.Key.MatchesWildcard(TimerWildcard))
+				if (TimerName.MatchesWildcard(Pattern))
 				{
-					OutIncludedTimers.Add(KeyValuePair.Value);
+					OutIncludedTimers.Add(Timer->Id);
 				}
 			}
 		}
diff --git a/Engine/Source/Developer/TraceInsights/Private/Insights/TimingProfiler/Widgets/STimerTreeView.cpp b/Engine/Source/Developer/TraceInsights/Private/Insights/TimingProfiler/Widgets/STimerTreeView.cpp
--- a/Engine/Source/Developer/TraceInsights/Private/Insights/TimingProfiler/Widgets/STimerTreeView.cpp
+++ b/Engine/Source/Developer/TraceInsights/Private/Insights/TimingProfiler/Widgets/STimerTreeView.cpp
@@ -1498,34 +1498,33 @@ void STimerTreeView::SetTree(const TraceServices::FTimingProfilerButterflyNode&
 	FTimerNodePtr RootTimerNodePtr;
 
 	TSharedPtr<const TraceServices::IAnalysisSession> Session = FInsightsManager::Get()->GetSession();
-	if (!Session.IsValid())
+	if (Session.IsValid())
 	{
-		return;
-	}
-	const TraceServices::ITimingProfilerProvider* TimingProfilerProvider = TraceServices::ReadTimingProfilerProvider(*Session.Get());
-	if (TimingProfilerProvider)
-	{
-		RootTimerNodePtr = CreateTimerNodeRec(Root, *TimingProfilerProvider);
-		if (RootTimerNodePtr)
+		const TraceServices::ITimingProfilerProvider* TimingProfilerProvider = TraceServices::ReadTimingProfilerProvider(*Session.Get());
+		if (TimingProfilerProvider)
 		{
-			// Mark the hot path. The child nodes are already sorted by InclTime (descending), so we just follow the first child.
-			FTimerNodePtr TimerNodePtr = RootTimerNodePtr;
-			while (TimerNodePtr.IsValid())
+			RootTimerNodePtr = CreateTimerNodeRec(Root, *TimingProfilerProvider);
+			if (RootTimerNodePtr)
 			{
-				TimerNodePtr->SetIsHotPath(true);
-				const TArray<FBaseTreeNodePtr>& Children = TimerNodePtr->GetChildren();
-				if (Children.Num() > 0)
-				{
-					check(Children[0]->Is<FTimerNode>());
-					TimerNodePtr = StaticCastSharedPtr<FTimerNode>(Children[0]);
-				}
-				else
+				// Mark the hot path. The child nodes are already sorted by InclTime (descending), so we just follow the first child.
+				FTimerNodePtr TimerNodePtr = RootTimerNodePtr;
+				while (TimerNodePtr.IsValid())
 				{
-					TimerNodePtr = nullptr;
+					TimerNodePtr->SetIsHotPath(true);
+					const TArray<FBaseTreeNodePtr>& Children = TimerNodePtr->GetChildren();
+					if (Children.Num() > 0)
+					{
+						check(Children[0]->Is<FTimerNode>());
+						TimerNodePtr = StaticCastSharedPtr<FTimerNode>(Children[0]);
+					}
+					else
+					{
+						TimerNodePtr = nullptr;
+					}
 				}
-			}
 
-			TreeNodes.Add(RootTimerNodePtr);
+				TreeNodes.Add(RootTimerNodePtr);
+			}
 		}
 	}
 
diff --git a/Engine/Source/Developer/TraceInsights/Private/Insights/TimingProfiler/Widgets/STimersView.cpp b/Engine/Source/Developer/TraceInsights/Private/Insights/TimingProfiler/Widgets/STimersView.cpp
--- a/Engine/Source/Developer/TraceInsights/Private/Insights/TimingProfiler/Widgets/STimersView.cpp
+++ b/Engine/Source/Developer/TraceInsights/Private/Insights/TimingProfiler/Widgets/STimersView.cpp
@@ -2773,7 +2773,8 @@ void STimersView::Tick(const FGeometry& AllottedGeometry, const double InCurrent
 
 	if (Session.IsValid())
 	{
-		if (const TraceServices::ITimingProfilerProvider* TimingProfilerProvider = TraceServices::ReadTimingProfilerProvider(*Session.Get()))
+		const TraceServices::ITimingProfilerProvider* TimingProfilerProvider = TraceServices::ReadTimingProfilerProvider(*Session.Get());
+		if (TimingProfilerProvider)
 		{
 			TraceServices::FAnalysisSessionReadScope SessionReadScope(*Session.Get());
 
@@ -2788,7 +2789,8 @@ void STimersView::Tick(const FGeometry& AllottedGeometry, const double InCurrent
 			bTimerTypeIsEnabled[static_cast<int>(ETimerNodeType::CpuScope)] = false;
 		}
 
-		if (const TraceServices::IStackSamplesProvider* StackSamplesProvider = TraceServices::ReadStackSamplesProvider(*Session.Get()))
+		const TraceServices::IStackSamplesProvider* StackSamplesProvider = TraceServices::ReadStackSamplesProvider(*Session.Get());
+		if (StackSamplesProvider)
 		{
 			TraceServices::FProviderReadScopeLock _(*StackSamplesProvider);
 
@@ -2823,73 +2825,80 @@ void STimersView::RebuildTree(bool bResync)
 		TimerNodes.Empty();
 	}
 
-	const uint32 PreviousNodeCount = TimerNodes.Num();
+	uint32 AddedNodeCount = 0;
 
-	if (Session.IsValid() && TraceServices::ReadTimingProfilerProvider(*Session.Get()))
+	if (Session.IsValid())
 	{
-		TraceServices::FAnalysisSessionReadScope SessionReadScope(*Session.Get());
-
-		const TraceServices::ITimingProfilerProvider& TimingProfilerProvider = *TraceServices::ReadTimingProfilerProvider(*Session.Get());
-
-		const TraceServices::ITimingProfilerTimerReader* TimerReader;
-		TimingProfilerProvider.ReadTimers([&TimerReader](const TraceServices::ITimingProfilerTimerReader& Out) { TimerReader = &Out; });
-
-		const uint32 TimerCount = TimerReader->GetTimerCount();
-		if (TimerCount != PreviousNodeCount)
+		const TraceServices::ITimingProfilerProvider* TimingProfilerProvider = TraceServices::ReadTimingProfilerProvider(*Session.Get());
+		if (TimingProfilerProvider)
 		{
-			check(TimerCount > PreviousNodeCount);
-			TimerNodes.Reserve(TimerCount);
+			TraceServices::FAnalysisSessionReadScope SessionReadScope(*Session.Get());
 
-			TSharedPtr<SFrameTrack> FrameTrack = GetFrameTrack();
-			TSharedPtr<FTimingGraphTrack> GraphTrack = GetTimingViewMainGraphTrack();
+			const TraceServices::ITimingProfilerTimerReader& TimerReader = TimingProfilerProvider->GetTimerReader();
 
-			// Add nodes only for new timers.
-			for (uint32 TimerIndex = PreviousNodeCount; TimerIndex < TimerCount; ++TimerIndex)
+			const uint32 PreviousNodeCount = TimerNodes.Num();
+			const uint32 TimerCount = TimerReader.GetTimerCount();
+			if (TimerCount != PreviousNodeCount)
 			{
-				const TraceServices::FTimingProfilerTimer& Timer = *(TimerReader->GetTimer(TimerIndex));
-				ensure(Timer.Id == TimerIndex);
-				const ETimerNodeType Type =
-					Timer.Type == TraceServices::ETimingProfilerTimerType::GpuScope ? ETimerNodeType::GpuScope :
-					Timer.Type == TraceServices::ETimingProfilerTimerType::VerseSampling ? ETimerNodeType::VerseSampling :
-					Timer.Type == TraceServices::ETimingProfilerTimerType::CpuScope ? ETimerNodeType::CpuScope :
-					Timer.Type == TraceServices::ETimingProfilerTimerType::CpuSampling ? ETimerNodeType::CpuSampling :
-					ETimerNodeType::InvalidOrMax;
-
-				FTimerNodePtr TimerNodePtr;
-				TStringBuilder<1024> StringBuilder;
-				if (FInsightsManager::Get()->GetSettings().GetShowFormatInTimerNames())
-				{
-					TimingProfilerProvider.MakeTimerDisplayName(Timer, StringBuilder);
-				}
-				else
+				check(TimerCount > PreviousNodeCount);
+				TimerNodes.Reserve(TimerCount);
+
+				TSharedPtr<SFrameTrack> FrameTrack = GetFrameTrack();
+				TSharedPtr<FTimingGraphTrack> GraphTrack = GetTimingViewMainGraphTrack();
+
+				// Add nodes only for new timers.
+				for (uint32 TimerIndex = PreviousNodeCount; TimerIndex < TimerCount; ++TimerIndex)
 				{
-					StringBuilder.Append(Timer.Name);
-				}
+					const TraceServices::FTimingProfilerTimer* TimerPtr = TimerReader.GetTimer(TimerIndex);
+					check(TimerPtr != nullptr);
+					const TraceServices::FTimingProfilerTimer& Timer = *TimerPtr;
+
+					ensure(Timer.Id == TimerIndex);
+
+					const ETimerNodeType Type =
+						Timer.Type == TraceServices::ETimingProfilerTimerType::GpuScope ? ETimerNodeType::GpuScope :
+						Timer.Type == TraceServices::ETimingProfilerTimerType::VerseSampling ? ETimerNodeType::VerseSampling :
+						Timer.Type == TraceServices::ETimingProfilerTimerType::CpuScope ? ETimerNodeType::CpuScope :
+						Timer.Type == TraceServices::ETimingProfilerTimerType::CpuSampling ? ETimerNodeType::CpuSampling :
+						ETimerNodeType::InvalidOrMax;
+
+					FTimerNodePtr TimerNodePtr;
+					TStringBuilder<1024> StringBuilder;
+					if (FInsightsManager::Get()->GetSettings().GetShowFormatInTimerNames())
+					{
+						TimingProfilerProvider->MakeTimerDisplayName(Timer, StringBuilder);
+					}
+					else
+					{
+						StringBuilder.Append(Timer.Name);
+					}
 
-				TimerNodePtr = MakeShared<FTimerNode>(Timer.Id, *StringBuilder, Type, false);
-				TimerNodePtr->SetDefaultSortOrder(TimerIndex + 1);
+					TimerNodePtr = MakeShared<FTimerNode>(Timer.Id, *StringBuilder, Type, false);
+					TimerNodePtr->SetDefaultSortOrder(TimerIndex + 1);
 
-				const bool bIsTimerAddedToGraphs =
-					(FrameTrack.IsValid() && FrameTrack->HasAnySeriesForTimer(Timer.Id)) ||
-					(GraphTrack.IsValid() && GraphTrack->HasAnySeriesForTimer(Timer.Id));
-				TimerNodePtr->SetAddedToGraphsFlag(bIsTimerAddedToGraphs);
+					const bool bIsTimerAddedToGraphs =
+						(FrameTrack.IsValid() && FrameTrack->HasAnySeriesForTimer(Timer.Id)) ||
+						(GraphTrack.IsValid() && GraphTrack->HasAnySeriesForTimer(Timer.Id));
+					TimerNodePtr->SetAddedToGraphsFlag(bIsTimerAddedToGraphs);
 
-				FLinearColor ColorOverride;
-				TSharedPtr<STimingProfilerWindow> ProfilerWindow = FTimingProfilerManager::Get()->GetProfilerWindow();
-				if (ProfilerWindow.IsValid() && ProfilerWindow->GetTimerColor(Timer.Id, ColorOverride))
-				{
-					TimerNodePtr->SetColor(ColorOverride);
-				}
+					FLinearColor ColorOverride;
+					TSharedPtr<STimingProfilerWindow> ProfilerWindow = FTimingProfilerManager::Get()->GetProfilerWindow();
+					if (ProfilerWindow.IsValid() && ProfilerWindow->GetTimerColor(Timer.Id, ColorOverride))
+					{
+						TimerNodePtr->SetColor(ColorOverride);
+					}
 
-				TimerNodes.Add(TimerNodePtr);
+					TimerNodes.Add(TimerNodePtr);
+					++AddedNodeCount;
+				}
+				ensure(TimerNodes.Num() == TimerCount);
 			}
-			ensure(TimerNodes.Num() == TimerCount);
 		}
 	}
 
 	SyncStopwatch.Stop();
 
-	if (bResync || TimerNodes.Num() != PreviousNodeCount)
+	if (bResync || AddedNodeCount > 0)
 	{
 		// Disable sorting if too many items.
 		if (TimerNodes.Num() > 10000)
@@ -2932,7 +2941,7 @@ void STimersView::RebuildTree(bool bResync)
 	{
 		const double SyncTime = SyncStopwatch.GetAccumulatedTime();
 		UE_LOG(LogTimingProfiler, Log, TEXT("[Timers] Tree view rebuilt in %.4fs (sync: %.4fs + update: %.4fs) --> %d timers (%d added)"),
-			TotalTime, SyncTime, TotalTime - SyncTime, TimerNodes.Num(), TimerNodes.Num() - PreviousNodeCount);
+			TotalTime, SyncTime, TotalTime - SyncTime, TimerNodes.Num(), AddedNodeCount);
 	}
 }
 
@@ -3512,30 +3521,35 @@ void STimersView::ContextMenu_ExportTimingEventsSelection_Execute() const
 	if (TimingView.IsValid())
 	{
 		// Add available GPU threads to the ExcludedThreads list.
-		ExcludedThreads.Add(FGpuTimingTrack::Gpu1ThreadId);
-		ExcludedThreads.Add(FGpuTimingTrack::Gpu2ThreadId);
+		{
+			ExcludedThreads.Add(FGpuTimingTrack::Gpu1ThreadId);
+			ExcludedThreads.Add(FGpuTimingTrack::Gpu2ThreadId);
+		}
 
-		// Add Gpu Queues for the new Gpu Profiler
+		// Add GPU Queues for the new GPU Profiler
+		const TraceServices::ITimingProfilerProvider* TimingProfilerProvider = TraceServices::ReadTimingProfilerProvider(*Session.Get());
+		if (TimingProfilerProvider)
 		{
 			TraceServices::FAnalysisSessionReadScope SessionReadScope(*Session.Get());
 
-			const TraceServices::ITimingProfilerProvider* TimingProfilerProvider = TraceServices::ReadTimingProfilerProvider(*Session.Get());
-			if (TimingProfilerProvider)
-			{
-				TimingProfilerProvider->EnumerateGpuQueues([&ExcludedThreads](const TraceServices::FGpuQueueInfo& QueueInfo)
-					{
-						ExcludedThreads.Add(QueueInfo.Id + (1 << 16));
-					});
-			}
+			TimingProfilerProvider->EnumerateGpuQueues(
+				[&ExcludedThreads](const TraceServices::FGpuQueueInfo& QueueInfo)
+				{
+					ExcludedThreads.Add(QueueInfo.Id + (1 << 16));
+				});
 		}
 
 		// Add the Verse Sampling fake thread to the ExcludedThreads list.
-		ExcludedThreads.Add(FVerseTimingTrack::VerseThreadId);
+		{
+			ExcludedThreads.Add(FVerseTimingTrack::VerseThreadId);
+		}
 
 		// Add available CPU threads to the ExcludedThreads list.
 		{
 			TraceServices::FAnalysisSessionReadScope SessionReadScope(*Session.Get());
+
 			const TraceServices::IThreadProvider& ThreadProvider = TraceServices::ReadThreadProvider(*Session.Get());
+
 			ThreadProvider.EnumerateThreads(
 				[&ExcludedThreads](const TraceServices::FThreadInfo& ThreadInfo)
 				{
@@ -3544,7 +3558,9 @@ void STimersView::ContextMenu_ExportTimingEventsSelection_Execute() const
 		}
 
 		// Move the threads corresponding to visible CPU/GPU tracks to the IncludedThreads list.
-		TimingView->EnumerateAllTracks([&IncludedThreads, &ExcludedThreads](TSharedPtr<FBaseTimingTrack>& Track) -> bool
+		TimingView->EnumerateAllTracks(
+			[&IncludedThreads, &ExcludedThreads]
+			(TSharedPtr<FBaseTimingTrack>& Track) -> bool
 			{
 				if (Track->IsVisible())
 				{
diff --git a/Engine/Source/Developer/TraceInsights/Private/Insights/Widgets/STimingView.cpp b/Engine/Source/Developer/TraceInsights/Private/Insights/Widgets/STimingView.cpp
--- a/Engine/Source/Developer/TraceInsights/Private/Insights/Widgets/STimingView.cpp
+++ b/Engine/Source/Developer/TraceInsights/Private/Insights/Widgets/STimingView.cpp
@@ -6163,19 +6163,22 @@ void STimingView::PopulateTrackSuggestionList(const FString& Text, TArray<FStrin
 void STimingView::PopulateTimerNameSuggestionList(const FString& Text, TArray<FString>& OutSuggestions)
 {
 	TSharedPtr<const TraceServices::IAnalysisSession> Session = FInsightsManager::Get()->GetSession();
-	if (Session.IsValid() && TraceServices::ReadTimingProfilerProvider(*Session.Get()))
+	if (!Session.IsValid())
 	{
-		TraceServices::FAnalysisSessionReadScope SessionReadScope(*Session.Get());
+		return;
+	}
 
-		const TraceServices::ITimingProfilerProvider& TimingProfilerProvider = *TraceServices::ReadTimingProfilerProvider(*Session.Get());
+	const TraceServices::ITimingProfilerProvider* TimingProfilerProvider = TraceServices::ReadTimingProfilerProvider(*Session.Get());
+	if (TimingProfilerProvider)
+	{
+		TraceServices::FAnalysisSessionReadScope SessionReadScope(*Session.Get());
 
-		const TraceServices::ITimingProfilerTimerReader* TimerReader;
-		TimingProfilerProvider.ReadTimers([&TimerReader](const TraceServices::ITimingProfilerTimerReader& Out) { TimerReader = &Out; });
+		const TraceServices::ITimingProfilerTimerReader& TimerReader = TimingProfilerProvider->GetTimerReader();
 
-		uint32 TimerCount = TimerReader->GetTimerCount();
+		uint32 TimerCount = TimerReader.GetTimerCount();
 		for (uint32 TimerIndex = 0; TimerIndex < TimerCount; ++TimerIndex)
 		{
-			const TraceServices::FTimingProfilerTimer* Timer = TimerReader->GetTimer(TimerIndex);
+			const TraceServices::FTimingProfilerTimer* Timer = TimerReader.GetTimer(TimerIndex);
 			if (Timer && Timer->Name)
 			{
 				if (Text.IsEmpty())
diff --git a/Engine/Source/Developer/TraceServices/Private/Common/Utils.cpp b/Engine/Source/Developer/TraceServices/Private/Common/Utils.cpp
--- a/Engine/Source/Developer/TraceServices/Private/Common/Utils.cpp
+++ b/Engine/Source/Developer/TraceServices/Private/Common/Utils.cpp
@@ -1,11 +1,14 @@
 // Copyright Epic Games, Inc. All Rights Reserved.
 
 #include "Utils.h"
-#include "TraceServices/Utils.h"
+
+#include "HAL/FileManager.h"
+#include "Templates/SharedPointer.h"
+
+// TraceServices
 #include "Common/FormatArgs.h"
+#include "TraceServices/Utils.h"
 #include "TraceServices/Containers/Tables.h"
-#include "Templates/SharedPointer.h"
-#include "HAL/FileManager.h"
 
 DEFINE_LOG_CATEGORY(LogTraceServices);
 
diff --git a/Engine/Source/Developer/TraceServices/Private/Model/TimingProfilerPrivate.h b/Engine/Source/Developer/TraceServices/Private/Model/TimingProfilerPrivate.h
--- a/Engine/Source/Developer/TraceServices/Private/Model/TimingProfilerPrivate.h
+++ b/Engine/Source/Developer/TraceServices/Private/Model/TimingProfilerPrivate.h
@@ -42,6 +42,7 @@ public:
 	// Timers
 
 	virtual void ReadTimers(TFunctionRef<void(const ITimingProfilerTimerReader&)> Callback) const override;
+	virtual const ITimingProfilerTimerReader& GetTimerReader() const override { return *this; }
 
 	// Metadata
 
diff --git a/Engine/Source/Developer/TraceServices/Public/TraceServices/Model/TimingProfiler.h b/Engine/Source/Developer/TraceServices/Public/TraceServices/Model/TimingProfiler.h
--- a/Engine/Source/Developer/TraceServices/Public/TraceServices/Model/TimingProfiler.h
+++ b/Engine/Source/Developer/TraceServices/Public/TraceServices/Model/TimingProfiler.h
@@ -287,6 +287,20 @@ public:
 
 	virtual void ReadTimers(TFunctionRef<void(const ITimingProfilerTimerReader&)> Callback) const = 0;
 
+	/**
+	 * Gets the TimerReader provider.
+	 *
+	 * @return A reference to the ITimingProfilerTimerReader object.
+	 *         The reference is valid for the rest of the session lifetime, but it can be used only under a Session Read Scope.
+	 */
+	virtual const ITimingProfilerTimerReader& GetTimerReader() const
+	{
+		const ITimingProfilerTimerReader* TimerReader = nullptr;
+		ReadTimers([&TimerReader](const TraceServices::ITimingProfilerTimerReader& Out) { TimerReader = &Out; });
+		check(TimerReader != nullptr);
+		return *TimerReader;
+	}
+
 	//////////////////////////////////////////////////
 	// Metadata
 
@@ -294,7 +308,7 @@ public:
 	virtual TArrayView<const uint8> GetMetadata(uint32 MetadataTimerId) const { return TArrayView<const uint8>(); } // also in ITimingProfilerTimerReader
 
 	/**
-	 * Get the metadata spec associated with the MetadataSpecId.
+	 * Gets the metadata spec associated with the MetadataSpecId.
 	 *
 	 * @param MetadataSpecId The Metadata spec id.
 	 * @return A pointer to the FMetadataSpec or nullptr if no spec is associated with the provided Id. The pointer is only valid in the same Session Read Scope as the function call.
