Insights Color Mode "By [REDACTED_PERSON]"
- Determines timer color based on both the timer name and the associated metadata values. Useful for differentiating between adjacent frames with the "Frame N" markers.

#Synced-CL 50397505

diff --git a/Engine/Source/Developer/TraceInsights/Private/Insights/TimingProfiler/TimingProfilerManager.h b/Engine/Source/Developer/TraceInsights/Private/Insights/TimingProfiler/TimingProfilerManager.h
--- a/Engine/Source/Developer/TraceInsights/Private/Insights/TimingProfiler/TimingProfilerManager.h
+++ b/Engine/Source/Developer/TraceInsights/Private/Insights/TimingProfiler/TimingProfilerManager.h
@@ -19,6 +19,7 @@ namespace UE::Insights
 	enum class ETimingEventsColoringMode : uint32
 	{
 		ByTimerName,
+		ByTimerNameAndMetadata,
 		ByTimerId,
 		BySourceFile,
 		ByDuration,
diff --git a/Engine/Source/Developer/TraceInsights/Private/Insights/TimingProfiler/Tracks/ThreadTimingTrack.cpp b/Engine/Source/Developer/TraceInsights/Private/Insights/TimingProfiler/Tracks/ThreadTimingTrack.cpp
--- a/Engine/Source/Developer/TraceInsights/Private/Insights/TimingProfiler/Tracks/ThreadTimingTrack.cpp
+++ b/Engine/Source/Developer/TraceInsights/Private/Insights/TimingProfiler/Tracks/ThreadTimingTrack.cpp
@@ -642,6 +642,11 @@ void AddTimingEventToBuilder(
 		case ETimingEventsColoringMode::ByTimerName:
 			EventColor = FTimingEvent::ComputeEventColor(Timer->Name);
 			break;
+		case ETimingEventsColoringMode::ByTimerNameAndMetadata:
+			EventColor = int32(TimerIndex) < 0 // has metadata?
+				? FTimingEvent::ComputeEventColor(Timer->Name, TimerReader.GetMetadata(TimerIndex))
+				: FTimingEvent::ComputeEventColor(Timer->Name);
+			break;
 		case ETimingEventsColoringMode::ByTimerId:
 			EventColor = FTimingEvent::ComputeEventColor(Timer->Id);
 			break;
diff --git a/Engine/Source/Developer/TraceInsights/Private/Insights/ViewModels/TimingEvent.cpp b/Engine/Source/Developer/TraceInsights/Private/Insights/ViewModels/TimingEvent.cpp
--- a/Engine/Source/Developer/TraceInsights/Private/Insights/ViewModels/TimingEvent.cpp
+++ b/Engine/Source/Developer/TraceInsights/Private/Insights/ViewModels/TimingEvent.cpp
@@ -37,6 +37,28 @@ uint32 FTimingEvent::ComputeEventColor(const TCHAR* Str)
 	return Color | 0xFF000000;
 }
 
+////////////////////////////////////////////////////////////////////////////////////////////////////
+
+uint32 FTimingEvent::ComputeEventColor(const TCHAR* Str, TConstArrayView<uint8> Metadata)
+{
+	uint32 Color = 0;
+	// Let the special event names (ex. "<unknown>", "<noname>", "<invalid>", etc.) to be colored in black.
+	if (Str != nullptr && *Str != TEXT('<'))
+	{
+		for (const TCHAR* c = Str; *c; ++c)
+		{
+			Color = (Color + *c) * 0x2c2c57ed;
+		}
+	}
+
+	for (int Index = 0; Index < Metadata.Num(); ++Index)
+	{
+		Color = (Color + Metadata[Index]) * 0x2c2c57ed;
+	}
+
+	return Color | 0xFF000000;
+}
+
 ////////////////////////////////////////////////////////////////////////////////////////////////////
 // ITimingEventFilter, FTimingEventFilter, ...
 ////////////////////////////////////////////////////////////////////////////////////////////////////
diff --git a/Engine/Source/Developer/TraceInsights/Private/Insights/Widgets/STimingView.cpp b/Engine/Source/Developer/TraceInsights/Private/Insights/Widgets/STimingView.cpp
--- a/Engine/Source/Developer/TraceInsights/Private/Insights/Widgets/STimingView.cpp
+++ b/Engine/Source/Developer/TraceInsights/Private/Insights/Widgets/STimingView.cpp
@@ -5331,6 +5331,17 @@ void STimingView::CreateCpuThreadTrackColoringModeMenu(FMenuBuilder& MenuBuilder
 			NAME_None,
 			EUserInterfaceActionType::RadioButton
 		);
+		MenuBuilder.AddMenuEntry(
+			LOCTEXT("CpuThreadTrackColoringMode_ByTimerNameAndMetadata", "By Timer Name And Metadata"),
+			LOCTEXT("CpuThreadTrackColoringMode_ByTimerNameAndMetadata_Desc", "Assign a color to CPU/GPU timing events based on their combined timer name and metadata values."),
+			FSlateIcon(),
+			FUIAction(
+				FExecuteAction::CreateSP(this, &STimingView::SetCpuThreadTrackColoringMode, ETimingEventsColoringMode::ByTimerNameAndMetadata),
+				FCanExecuteAction(),
+				FIsActionChecked::CreateSP(this, &STimingView::CheckCpuThreadTrackColoringMode, ETimingEventsColoringMode::ByTimerNameAndMetadata)),
+			NAME_None,
+			EUserInterfaceActionType::RadioButton
+		);
 		MenuBuilder.AddMenuEntry(
 			LOCTEXT("CpuThreadTrackColoringMode_ByTimerId", "By Timer Id"),
 			LOCTEXT("CpuThreadTrackColoringMode_ByTimerId_Desc", "Assign a color to CPU/GPU timing events based on their timer id."),
diff --git a/Engine/Source/Developer/TraceInsights/Public/Insights/ViewModels/TimingEvent.h b/Engine/Source/Developer/TraceInsights/Public/Insights/ViewModels/TimingEvent.h
--- a/Engine/Source/Developer/TraceInsights/Public/Insights/ViewModels/TimingEvent.h
+++ b/Engine/Source/Developer/TraceInsights/Public/Insights/ViewModels/TimingEvent.h
@@ -80,6 +80,7 @@ public:
 
 	static UE_API uint32 ComputeEventColor(uint32 Id);
 	static UE_API uint32 ComputeEventColor(const TCHAR* Str);
+	static UE_API uint32 ComputeEventColor(const TCHAR* Str, TConstArrayView<uint8> Metadata);
 
 protected:
 	void SetType(uint64 InType) { Type = InType; }
