[Insights]
- Timers/Counters: Improved the responsiveness for context menus when a large number of timers/counters is selected.
- Normalized code between STimersView and SStatsView.

#Synced-CL 50483080

diff --git a/Engine/Source/Developer/TraceInsights/Private/Insights/TimingProfiler/Widgets/SStatsView.cpp b/Engine/Source/Developer/TraceInsights/Private/Insights/TimingProfiler/Widgets/SStatsView.cpp
--- a/Engine/Source/Developer/TraceInsights/Private/Insights/TimingProfiler/Widgets/SStatsView.cpp
+++ b/Engine/Source/Developer/TraceInsights/Private/Insights/TimingProfiler/Widgets/SStatsView.cpp
@@ -37,6 +37,7 @@
 // TraceInsights
 #include "Insights/InsightsStyle.h"
 #include "Insights/Log.h"
+#include "Insights/Table/ViewModels/TableCommands.h"
 #include "Insights/TimingProfiler/GraphTracks/TimingGraphSeries.h"
 #include "Insights/TimingProfiler/GraphTracks/TimingGraphTrack.h"
 #include "Insights/TimingProfiler/TimingProfilerManager.h"
@@ -49,7 +50,6 @@
 #include "Insights/TimingProfiler/Widgets/STimingProfilerWindow.h"
 #include "Insights/TimingProfilerCommon.h"
 #include "Insights/Widgets/STimingView.h"
-#include "Insights/Table/ViewModels/TableCommands.h"
 
 #include <limits>
 
@@ -148,11 +148,6 @@ public:
 
 SStatsView::SStatsView()
 	: Table(MakeShared<FTable>())
-	, bExpansionSaved(false)
-	, bFilterOutZeroCountStats(false)
-	, GroupingMode(EStatsGroupingMode::Flat)
-	, AvailableSorters()
-	, CurrentSorter(nullptr)
 	, ColumnBeingSorted(GetDefaultColumnBeingSorted())
 	, ColumnSortMode(GetDefaultColumnSortMode())
 	, Aggregator(MakeShared<FCounterAggregator>())
@@ -411,9 +406,8 @@ END_SLATE_FUNCTION_BUILD_OPTIMIZATION
 
 TSharedPtr<SWidget> SStatsView::TreeView_GetMenuContent()
 {
-	const TArray<FStatsNodePtr> SelectedNodes = TreeView->GetSelectedItems();
-	const int32 NumSelectedNodes = SelectedNodes.Num();
-	FStatsNodePtr SelectedNode = NumSelectedNodes ? SelectedNodes[0] : nullptr;
+	const int32 NumSelectedNodes = TreeView->GetNumItemsSelected();
+	FStatsNodePtr SingleSelectedNode;
 
 	FText SelectionStr;
 	if (NumSelectedNodes == 0)
@@ -422,7 +416,11 @@ TSharedPtr<SWidget> SStatsView::TreeView_GetMenuContent()
 	}
 	else if (NumSelectedNodes == 1)
 	{
-		FString ItemName = SelectedNode->GetName().ToString();
+		TArray<FStatsNodePtr> SelectedNodes = TreeView->GetSelectedItems();
+		check(SelectedNodes.Num() == 1);
+		SingleSelectedNode = SelectedNodes[0];
+
+		FString ItemName = SingleSelectedNode->GetName().ToString();
 		const int32 MaxStringLen = 64;
 		if (ItemName.Len() > MaxStringLen)
 		{
@@ -466,22 +464,23 @@ TSharedPtr<SWidget> SStatsView::TreeView_GetMenuContent()
 	// Counter options section
 	MenuBuilder.BeginSection("CounterOptions", LOCTEXT("ContextMenu_Section_CounterOptions", "Counter Options"));
 	{
-		auto CanExecute = [NumSelectedNodes, SelectedNode]()
+		const bool bIsValidNode =
+			SingleSelectedNode.IsValid() &&
+			SingleSelectedNode->GetType() != EStatsNodeType::Group;
+
+		auto CanExecute = [this, bIsValidNode]()
 		{
-			TSharedPtr<STimingProfilerWindow> Wnd = FTimingProfilerManager::Get()->GetProfilerWindow();
-			TSharedPtr<STimingView> TimingView = Wnd.IsValid() ? Wnd->GetTimingView() : nullptr;
-			return TimingView.IsValid() && NumSelectedNodes == 1 && SelectedNode.IsValid() && SelectedNode->GetType() != EStatsNodeType::Group;
+			return bIsValidNode && GetTimingView().IsValid();
 		};
 
 		// Add/remove series to/from graph track
 		{
 			FUIAction Action_ToggleCounterInGraphTrack;
 			Action_ToggleCounterInGraphTrack.CanExecuteAction = FCanExecuteAction::CreateLambda(CanExecute);
-			Action_ToggleCounterInGraphTrack.ExecuteAction = FExecuteAction::CreateSP(this, &SStatsView::ToggleTimingViewMainGraphEventSeries, SelectedNode);
+			Action_ToggleCounterInGraphTrack.ExecuteAction = FExecuteAction::CreateSP(this, &SStatsView::ToggleTimingViewMainGraphEventSeries, SingleSelectedNode);
 
-			if (SelectedNode.IsValid() &&
-				SelectedNode->GetType() != EStatsNodeType::Group &&
-				IsSeriesInTimingViewMainGraph(SelectedNode))
+			if (bIsValidNode &&
+				IsSeriesInTimingViewMainGraph(SingleSelectedNode))
 			{
 				MenuBuilder.AddMenuEntry
 				(
@@ -923,12 +922,15 @@ void SStatsView::ApplyFiltering()
 		int32 NumVisibleChildren = 0;
 		for (const FBaseTreeNodePtr& ChildPtr : GroupChildren)
 		{
-			const FStatsNodePtr& NodePtr = StaticCastSharedPtr<FStatsNode, FBaseTreeNode>(ChildPtr);
+			check(ChildPtr->Is<FStatsNode>());
+			const FStatsNodePtr& NodePtr = StaticCastSharedPtr<FStatsNode>(ChildPtr);
+
+			const bool bIsChildVisible =
+				(!bFilterOutZeroCountStats || NodePtr->GetAggregatedStats().Count > 0) &&
+				FilterByNodeType[static_cast<int>(NodePtr->GetType())] &&
+				FilterByDataType[static_cast<int>(NodePtr->GetDataType())] &&
+				Filters->PassesAllFilters(NodePtr);
 
-			const bool bIsChildVisible = (!bFilterOutZeroCountStats || NodePtr->GetAggregatedStats().Count > 0)
-									  && FilterByNodeType[static_cast<int>(NodePtr->GetType())]
-									  && FilterByDataType[static_cast<int>(NodePtr->GetDataType())]
-									  && Filters->PassesAllFilters(NodePtr);
 			if (bIsChildVisible)
 			{
 				// Add a child.
@@ -1007,7 +1009,8 @@ void SStatsView::ApplyFiltering()
 		const TArray<FBaseTreeNodePtr>& GroupChildren = GroupPtr->GetFilteredChildren();
 		for (const FBaseTreeNodePtr& ChildPtr : GroupChildren)
 		{
-			const FStatsNodePtr& NodePtr = StaticCastSharedPtr<FStatsNode, FBaseTreeNode>(ChildPtr);
+			check(ChildPtr->Is<FStatsNode>());
+			const FStatsNodePtr& NodePtr = StaticCastSharedPtr<FStatsNode>(ChildPtr);
 			const FAggregatedStats& NodeAggregatedStats = NodePtr->GetAggregatedStats();
 
 			if (NodeAggregatedStats.Count > 0)
@@ -1196,23 +1199,37 @@ void SStatsView::TreeView_OnSelectionChanged(FStatsNodePtr SelectedItem, ESelect
 {
 	if (SelectInfo != ESelectInfo::Direct)
 	{
-		//TArray<FStatsNodePtr> SelectedItems = TreeView->GetSelectedItems();
-		//if (SelectedItems.Num() == 1 && SelectedItems[0]->GetType() != EStatsNodeType::Group)
+		//FStatsNodePtr SelectedNode = GetSingleSelectedStatsNode();
+		//if (SelectedNode.IsValid() && SelectedNode->GetType() != EStatsNodeType::Group)
 		//{
-		//	FTimingProfilerManager::Get()->SetSelectedCounter(SelectedItems[0]->GetCounterId());
+		//	FTimingProfilerManager::Get()->SetSelectedCounter(SelectedNode->GetCounterId());
 		//}
 	}
 }
 
 ////////////////////////////////////////////////////////////////////////////////////////////////////
 
+FStatsNodePtr SStatsView::GetSingleSelectedStatsNode() const
+{
+	if (TreeView->GetNumItemsSelected() != 1)
+	{
+		return nullptr;
+	}
+	TArray<FStatsNodePtr> SelectedNodes = TreeView->GetSelectedItems();
+	check(SelectedNodes.Num() == 1);
+	return SelectedNodes[0];
+}
+
+////////////////////////////////////////////////////////////////////////////////////////////////////
+
 void SStatsView::TreeView_OnGetChildren(FStatsNodePtr InParent, TArray<FStatsNodePtr>& OutChildren)
 {
 	const TArray<FBaseTreeNodePtr>& Children = InParent->GetFilteredChildren();
 	OutChildren.Reset(Children.Num());
 	for (const FBaseTreeNodePtr& Child : Children)
 	{
-		OutChildren.Add(StaticCastSharedPtr<FStatsNode, FBaseTreeNode>(Child));
+		check(Child->Is<FStatsNode>());
+		OutChildren.Add(StaticCastSharedPtr<FStatsNode>(Child));
 	}
 }
 
@@ -1460,6 +1477,10 @@ void SStatsView::CreateGroups()
 		GroupNodeSet.KeySort([](const uint32& A, const uint32& B) { return A > B; }); // sort groups by order
 		GroupNodeSet.GenerateValueArray(GroupNodes);
 	}
+	else
+	{
+		unimplemented();
+	}
 
 	Stopwatch.Stop();
 	const double TotalTime = Stopwatch.GetAccumulatedTime();
@@ -1497,7 +1518,9 @@ void SStatsView::CreateGroupByOptionsSources()
 
 void SStatsView::GroupBy_OnSelectionChanged(TSharedPtr<EStatsGroupingMode> NewGroupingMode, ESelectInfo::Type SelectInfo)
 {
-	if (SelectInfo != ESelectInfo::Direct)
+	if (SelectInfo != ESelectInfo::Direct &&
+		NewGroupingMode.IsValid() &&
+		*NewGroupingMode != GroupingMode)
 	{
 		GroupingMode = *NewGroupingMode;
 
@@ -1511,6 +1534,7 @@ void SStatsView::GroupBy_OnSelectionChanged(TSharedPtr<EStatsGroupingMode> NewGr
 
 TSharedRef<SWidget> SStatsView::GroupBy_OnGenerateWidget(TSharedPtr<EStatsGroupingMode> InGroupingMode) const
 {
+	check(InGroupingMode.IsValid());
 	return SNew(STextBlock)
 		.Text(StatsNodeGroupingHelper::ToText(*InGroupingMode))
 		.ToolTipText(StatsNodeGroupingHelper::ToDescription(*InGroupingMode));
@@ -1614,7 +1638,8 @@ void SStatsView::SortTreeNodesRec(FStatsNode& Node, const ITableCellValueSorter&
 	{
 		if (ChildPtr->GetChildrenCount() > 0)
 		{
-			SortTreeNodesRec(*StaticCastSharedPtr<FTimerNode>(ChildPtr), Sorter);
+			check(ChildPtr->Is<FStatsNode>());
+			SortTreeNodesRec(*StaticCastSharedPtr<FStatsNode>(ChildPtr), Sorter);
 		}
 	}
 #endif
@@ -1998,7 +2023,7 @@ void SStatsView::RebuildTree(bool bResync)
 		StatsNodesIdMap.Empty();
 	}
 
-	const uint32 PreviousNodeCount = StatsNodes.Num();
+	uint32 AddedNodeCount = 0;
 
 	if (Session.IsValid())
 	{
@@ -2006,6 +2031,7 @@ void SStatsView::RebuildTree(bool bResync)
 
 		const TraceServices::ICounterProvider& CountersProvider = TraceServices::ReadCounterProvider(*Session.Get());
 
+		const uint32 PreviousNodeCount = StatsNodes.Num();
 		const uint32 CounterCount = static_cast<uint32>(CountersProvider.GetCounterCount());
 		if (CounterCount != PreviousNodeCount)
 		{
@@ -2017,49 +2043,53 @@ void SStatsView::RebuildTree(bool bResync)
 
 			// Add nodes only for new counters.
 			uint32 CounterIndex = PreviousNodeCount;
-			CountersProvider.EnumerateCounters([this, &CounterIndex](uint32 CounterId, const TraceServices::ICounter& Counter)
-			{
-				auto DeduceGroup = [](const TraceServices::ICounter& Counter)
+			CountersProvider.EnumerateCounters(
+				[this, &CounterIndex, &AddedNodeCount]
+				(uint32 CounterId, const TraceServices::ICounter& Counter)
 				{
-					static const FName MemoryGroup(TEXT("Memory"));
-					static const FName BandwidthGroup(TEXT("Bandwidth"));
-					static const FName PercentGroup(TEXT("Percent"));
-					static const FName MiscFloatGroup(TEXT("Misc_double"));
-					static const FName MiscInt64Group(TEXT("Misc_int64"));
-
-					switch (Counter.GetDisplayHint())	
+					auto DeduceGroup = [](const TraceServices::ICounter& Counter)
+					{
+						static const FName MemoryGroup(TEXT("Memory"));
+						static const FName BandwidthGroup(TEXT("Bandwidth"));
+						static const FName PercentGroup(TEXT("Percent"));
+						static const FName MiscFloatGroup(TEXT("Misc_double"));
+						static const FName MiscInt64Group(TEXT("Misc_int64"));
+
+						switch (Counter.GetDisplayHint())
+						{
+						case TraceServices::CounterDisplayHint_Memory:
+							return MemoryGroup;
+						case TraceServices::CounterDisplayHint_Bandwidth:
+							return BandwidthGroup;
+						case TraceServices::CounterDisplayHint_Percent:
+							return PercentGroup;
+						}
+						return Counter.IsFloatingPoint() ? MiscFloatGroup : MiscInt64Group;
+					};
+
+					FStatsNodePtr NodePtr = StatsNodesIdMap.FindRef(CounterId);
+					if (!NodePtr)
 					{
-					case TraceServices::CounterDisplayHint_Memory:
-						return MemoryGroup;
-					case TraceServices::CounterDisplayHint_Bandwidth:
-						return BandwidthGroup;
-					case TraceServices::CounterDisplayHint_Percent:
-						return PercentGroup;
+						FName Name(Counter.GetName());
+						const FName Group = Counter.GetGroup() ? Counter.GetGroup() : DeduceGroup(Counter);
+						const EStatsNodeType Type = EStatsNodeType::Counter;
+						const EStatsNodeDataType DataType = Counter.IsFloatingPoint() ? EStatsNodeDataType::Double : EStatsNodeDataType::Int64;
+						NodePtr = MakeShared<FStatsNode>(CounterId, Name, Group, Type, DataType);
+						NodePtr->SetDefaultSortOrder(++CounterIndex);
+						UpdateNode(NodePtr);
+						StatsNodes.Add(NodePtr);
+						StatsNodesIdMap.Add(CounterId, NodePtr);
+						++AddedNodeCount;
 					}
-					return Counter.IsFloatingPoint() ? MiscFloatGroup : MiscInt64Group;
-				};
-
-				FStatsNodePtr NodePtr = StatsNodesIdMap.FindRef(CounterId);
-				if (!NodePtr)
-				{
-					FName Name(Counter.GetName());
-					const FName Group = Counter.GetGroup() ? Counter.GetGroup() : DeduceGroup(Counter);
-					const EStatsNodeType Type = EStatsNodeType::Counter;
-					const EStatsNodeDataType DataType = Counter.IsFloatingPoint() ? EStatsNodeDataType::Double : EStatsNodeDataType::Int64;
-					NodePtr = MakeShared<FStatsNode>(CounterId, Name, Group, Type, DataType);
-					NodePtr->SetDefaultSortOrder(++CounterIndex);
-					UpdateNode(NodePtr);
-					StatsNodes.Add(NodePtr);
-					StatsNodesIdMap.Add(CounterId, NodePtr);
-				}
-			});
+				});
+			check(AddedNodeCount == CounterIndex - PreviousNodeCount);
 			ensure(StatsNodes.Num() == CounterCount);
 		}
 	}
 
 	SyncStopwatch.Stop();
 
-	if (bResync || StatsNodes.Num() != PreviousNodeCount)
+	if (bResync || AddedNodeCount > 0)
 	{
 		// Disable sorting if too many items.
 		if (StatsNodes.Num() > 10000)
@@ -2102,7 +2132,7 @@ void SStatsView::RebuildTree(bool bResync)
 	{
 		const double SyncTime = SyncStopwatch.GetAccumulatedTime();
 		UE_LOG(LogTimingProfiler, Log, TEXT("[Counters] Tree view rebuilt in %.4fs (sync: %.4fs + update: %.4fs) --> %d counters (%d added)"),
-			TotalTime, SyncTime, TotalTime - SyncTime, StatsNodes.Num(), StatsNodes.Num() - PreviousNodeCount);
+			TotalTime, SyncTime, TotalTime - SyncTime, StatsNodes.Num(), AddedNodeCount);
 	}
 }
 
@@ -2114,19 +2144,11 @@ void SStatsView::UpdateNode(FStatsNodePtr NodePtr)
 
 	if (!NodePtr->IsGroup())
 	{
-		TSharedPtr<STimingProfilerWindow> Wnd = FTimingProfilerManager::Get()->GetProfilerWindow();
-		if (Wnd.IsValid())
+		TSharedPtr<FTimingGraphTrack> GraphTrack = GetTimingViewMainGraphTrack();
+		if (GraphTrack.IsValid())
 		{
-			TSharedPtr<STimingView> TimingView = Wnd->GetTimingView();
-			if (TimingView.IsValid())
-			{
-				TSharedPtr<FTimingGraphTrack> GraphTrack = TimingView->GetMainTimingGraphTrack();
-				if (GraphTrack.IsValid())
-				{
-					TSharedPtr<FTimingGraphSeries> Series = GraphTrack->GetStatsCounterSeries(NodePtr->GetCounterId());
-					bAddedToGraphFlag = Series.IsValid();
-				}
-			}
+			TSharedPtr<FTimingGraphSeries> Series = GraphTrack->GetStatsCounterSeries(NodePtr->GetCounterId());
+			bAddedToGraphFlag = Series.IsValid();
 		}
 	}
 
@@ -2176,7 +2198,7 @@ void SStatsView::FinishAggregation()
 	UpdateTree(); // grouping + sorting + filtering
 
 	// Ensure the last selected item is visible.
-	const TArray<FStatsNodePtr> SelectedNodes = TreeView->GetSelectedItems();
+	TArray<FStatsNodePtr> SelectedNodes = TreeView->GetSelectedItems();
 	if (SelectedNodes.Num() > 0)
 	{
 		TreeView->RequestScrollIntoView(SelectedNodes.Last());
@@ -2206,9 +2228,7 @@ void SStatsView::SelectCounterNode(uint32 CounterId)
 
 TSharedPtr<FTimingGraphTrack> SStatsView::GetTimingViewMainGraphTrack() const
 {
-	TSharedPtr<STimingProfilerWindow> Wnd = FTimingProfilerManager::Get()->GetProfilerWindow();
-	TSharedPtr<STimingView> TimingView = Wnd.IsValid() ? Wnd->GetTimingView() : nullptr;
-
+	TSharedPtr<STimingView> TimingView = GetTimingView();
 	return TimingView.IsValid() ? TimingView->GetMainTimingGraphTrack() : nullptr;
 }
 
@@ -2238,7 +2258,6 @@ void SStatsView::ToggleGraphSeries(TSharedRef<FTimingGraphTrack> GraphTrack, FSt
 bool SStatsView::IsSeriesInTimingViewMainGraph(FStatsNodePtr CounterNode) const
 {
 	TSharedPtr<FTimingGraphTrack> GraphTrack = GetTimingViewMainGraphTrack();
-
 	if (GraphTrack.IsValid())
 	{
 		const uint32 CounterId = CounterNode->GetCounterId();
@@ -2262,147 +2281,12 @@ void SStatsView::ToggleTimingViewMainGraphEventSeries(FStatsNodePtr CounterNode)
 }
 
 ////////////////////////////////////////////////////////////////////////////////////////////////////
-// Generate and Edit Colors
-////////////////////////////////////////////////////////////////////////////////////////////////////
-
-bool SStatsView::ContextMenu_GenerateNewColor_CanExecute() const
-{
-	return AreSelectedCounterNodesValid();
-}
-
-////////////////////////////////////////////////////////////////////////////////////////////////////
-
-void SStatsView::ContextMenu_GenerateNewColor_Execute() const
-{
-	const TArray<FStatsNodePtr> SelectedNodes = TreeView->GetSelectedItems();
-	for (const FStatsNodePtr& SelectedNode : SelectedNodes)
-	{
-		if (SelectedNode->Is<FStatsNode>())
-		{
-			FStatsNode& StatsNode = SelectedNode->As<FStatsNode>();
-			SetRandomColorToNode(StatsNode);
-		}
-	}
-};
-
-////////////////////////////////////////////////////////////////////////////////////////////////////
-
-void SStatsView::SetColorToNode(FStatsNode& StatsNode, FLinearColor CustomColor) const
-{
-	StatsNode.SetColor(CustomColor);
-
-	const uint32 CounterId = StatsNode.GetCounterId();
-	const FLinearColor BorderColor(CustomColor.R + 0.4f, CustomColor.G + 0.4f, CustomColor.B + 0.4f, 1.0f);
-
-	TSharedPtr<FTimingGraphSeries> CounterSeries = GetTimingViewMainGraphTrack()->GetStatsCounterSeries(CounterId);
-	if (CounterSeries.IsValid())
-	{
-		CounterSeries->SetColor(CustomColor, BorderColor);
-	}
-}
-
-////////////////////////////////////////////////////////////////////////////////////////////////////
-
-void SStatsView::SetRandomColorToNode(FStatsNode& StatsNode) const
-{
-	uint64 Time = FPlatformTime::Cycles64();
-	uint32 Hash = (Time & 0xFFFFFFFF) ^ (Time >> 32);
-	Hash = ((Hash << 5) + Hash) * 0xfb23618f;
-
-	const uint8 H = Hash & 0xFF;
-	const uint8 S = 128 + ((Hash >> 8) & 0x7F);
-	const uint8 V = 128 + ((Hash >> 16) & 0x7F);
-	FLinearColor CustomColor = FLinearColor::MakeFromHSV8(H, S, V);
-	CustomColor.A = 1.0f;
-
-	StatsNode.SetColor(CustomColor);
-
-	const uint32 CounterId = StatsNode.GetCounterId();
-	const FLinearColor BorderColor(CustomColor.R + 0.4f, CustomColor.G + 0.4f, CustomColor.B + 0.4f, 1.0f);
-
-	TSharedPtr<FTimingGraphSeries> CounterSeries = GetTimingViewMainGraphTrack()->GetStatsCounterSeries(CounterId);
-	if (CounterSeries.IsValid())
-	{
-		CounterSeries->SetColor(CustomColor, BorderColor);
-	}
-}
-
-////////////////////////////////////////////////////////////////////////////////////////////////////
-
-bool SStatsView::ContextMenu_EditColor_CanExecute() const
-{
-	return AreSelectedCounterNodesValid();
-};
-
-////////////////////////////////////////////////////////////////////////////////////////////////////
-
-void SStatsView::ContextMenu_EditColor_Execute()
-{
-	EditableColorValue = FLinearColor(0.5f, 0.5f, 0.5f, 1.0f);
-	const TArray<FStatsNodePtr> SelectedNodes = TreeView->GetSelectedItems();
-	for (const FStatsNodePtr& SelectedNode : SelectedNodes)
-	{
-		if (SelectedNode->Is<FStatsNode>())
-		{
-			FStatsNode& StatsNode = SelectedNode->As<FStatsNode>();
-			EditableColorValue = StatsNode.GetColor();
-			break;
-		}
-	}
-
-	FColorPickerArgs PickerArgs;
-	{
-		PickerArgs.bUseAlpha = true;
-		PickerArgs.bOnlyRefreshOnMouseUp = false;
-		PickerArgs.bOnlyRefreshOnOk = false;
-		PickerArgs.bExpandAdvancedSection = false;
-		PickerArgs.OnColorCommitted = FOnLinearColorValueChanged::CreateSP(this, &SStatsView::SetEditableColor);
-		PickerArgs.OnColorPickerCancelled = FOnColorPickerCancelled::CreateSP(this, &SStatsView::ColorPickerCancelled);
-		PickerArgs.InitialColor = EditableColorValue;
-		PickerArgs.ParentWidget = SharedThis(this);
-	}
-
-	OpenColorPicker(PickerArgs);
-};
-
-////////////////////////////////////////////////////////////////////////////////////////////////////
-
-FLinearColor SStatsView::GetEditableColor() const
-{
-	return EditableColorValue;
-}
-
-////////////////////////////////////////////////////////////////////////////////////////////////////
-
-void SStatsView::SetEditableColor(FLinearColor NewColor)
-{
-	EditableColorValue = NewColor;
-
-	const TArray<FStatsNodePtr> SelectedNodes = TreeView->GetSelectedItems();
-	for (const FStatsNodePtr& SelectedNode : SelectedNodes)
-	{
-		if (SelectedNode->Is<FStatsNode>())
-		{
-			FStatsNode& StatsNode = SelectedNode->As<FStatsNode>();
-			SetColorToNode(StatsNode, EditableColorValue);
-		}
-	}
-}
-
-////////////////////////////////////////////////////////////////////////////////////////////////////
-
-void SStatsView::ColorPickerCancelled(FLinearColor OriginalColor)
-{
-	SetEditableColor(OriginalColor);
-}
-
+// Copy to Clipboard
 ////////////////////////////////////////////////////////////////////////////////////////////////////
 
 bool SStatsView::ContextMenu_CopySelectedToClipboard_CanExecute() const
 {
-	const TArray<FStatsNodePtr> SelectedNodes = TreeView->GetSelectedItems();
-
-	return SelectedNodes.Num() > 0;
+	return TreeView->GetNumItemsSelected() > 0;
 }
 
 ////////////////////////////////////////////////////////////////////////////////////////////////////
@@ -2420,7 +2304,7 @@ void SStatsView::ContextMenu_CopySelectedToClipboard_Execute()
 
 bool SStatsView::ContextMenu_CopySelectedNameToClipboard_CanExecute() const
 {
-	return TreeView->GetSelectedItems().Num() > 0;
+	return TreeView->GetNumItemsSelected() > 0;
 }
 
 ////////////////////////////////////////////////////////////////////////////////////////////////////
@@ -2434,12 +2318,13 @@ void SStatsView::ContextMenu_CopySelectedNameToClipboard_Execute()
 	}
 }
 
+////////////////////////////////////////////////////////////////////////////////////////////////////
+// Export
 ////////////////////////////////////////////////////////////////////////////////////////////////////
 
 bool SStatsView::ContextMenu_Export_CanExecute() const
 {
-	const TArray<FStatsNodePtr> SelectedNodes = TreeView->GetSelectedItems();
-	return SelectedNodes.Num() > 0;
+	return TreeView->GetNumItemsSelected() > 0;
 }
 
 ////////////////////////////////////////////////////////////////////////////////////////////////////
@@ -2451,15 +2336,16 @@ void SStatsView::ContextMenu_Export_Execute()
 		return;
 	}
 
-	TArray<FBaseTreeNodePtr> SelectedNodes;
-	for (FStatsNodePtr Item : TreeView->GetSelectedItems())
+	TArray<FStatsNodePtr> SelectedStatsNodes = TreeView->GetSelectedItems();
+	if (SelectedStatsNodes.Num() == 0)
 	{
-		SelectedNodes.Add(Item);
+		return;
 	}
 
-	if (SelectedNodes.Num() == 0)
+	TArray<FBaseTreeNodePtr> SelectedNodes;
+	for (FStatsNodePtr Node : SelectedStatsNodes)
 	{
-		return;
+		SelectedNodes.Add(Node);
 	}
 
 	const FString DialogTitle = LOCTEXT("Export_Title", "Export Aggregated Counter Stats").ToString();
@@ -2580,8 +2466,8 @@ void SStatsView::ContextMenu_Export_Execute()
 
 bool SStatsView::ContextMenu_ExportValues_CanExecute() const
 {
-	const TArray<FStatsNodePtr> SelectedNodes = TreeView->GetSelectedItems();
-	return (SelectedNodes.Num() == 1) && !SelectedNodes[0]->IsGroup();
+	FStatsNodePtr SelectedNode = GetSingleSelectedStatsNode();
+	return SelectedNode.IsValid() && !SelectedNode->IsGroup();
 }
 
 ////////////////////////////////////////////////////////////////////////////////////////////////////
@@ -2593,14 +2479,14 @@ void SStatsView::ContextMenu_ExportValues_Execute() const
 		return;
 	}
 
-	const TArray<FStatsNodePtr> SelectedNodes = TreeView->GetSelectedItems();
-	if (SelectedNodes.Num() != 1 || SelectedNodes[0]->IsGroup())
+	FStatsNodePtr SelectedNode = GetSingleSelectedStatsNode();
+	if (!SelectedNode.IsValid() || SelectedNode->IsGroup())
 	{
 		return;
 	}
 
-	const uint32 CounterId = SelectedNodes[0]->GetCounterId();
-	FString CounterName = SelectedNodes[0]->GetName().ToString();
+	const uint32 CounterId = SelectedNode->GetCounterId();
+	FString CounterName = SelectedNode->GetName().ToString();
 
 	const FString DialogTitle = LOCTEXT("ExportValues_Title", "Export Counter Values").ToString();
 
@@ -2628,8 +2514,7 @@ void SStatsView::ContextMenu_ExportValues_Execute() const
 	Params.IntervalStartTime = -std::numeric_limits<double>::infinity();
 	Params.IntervalEndTime = +std::numeric_limits<double>::infinity();
 
-	TSharedPtr<STimingProfilerWindow> Wnd = FTimingProfilerManager::Get()->GetProfilerWindow();
-	TSharedPtr<STimingView> TimingView = Wnd.IsValid() ? Wnd->GetTimingView() : nullptr;
+	TSharedPtr<STimingView> TimingView = GetTimingView();
 	if (TimingView.IsValid())
 	{
 		const double SelectionStartTime = TimingView->GetSelectionStartTime();
@@ -2650,8 +2535,8 @@ void SStatsView::ContextMenu_ExportValues_Execute() const
 
 bool SStatsView::ContextMenu_ExportOps_CanExecute() const
 {
-	const TArray<FStatsNodePtr> SelectedNodes = TreeView->GetSelectedItems();
-	return (SelectedNodes.Num() == 1) && !SelectedNodes[0]->IsGroup();
+	FStatsNodePtr SelectedNode = GetSingleSelectedStatsNode();
+	return SelectedNode.IsValid() && !SelectedNode->IsGroup();
 }
 
 ////////////////////////////////////////////////////////////////////////////////////////////////////
@@ -2663,14 +2548,14 @@ void SStatsView::ContextMenu_ExportOps_Execute() const
 		return;
 	}
 
-	const TArray<FStatsNodePtr> SelectedNodes = TreeView->GetSelectedItems();
-	if (SelectedNodes.Num() != 1 || SelectedNodes[0]->IsGroup())
+	FStatsNodePtr SelectedNode = GetSingleSelectedStatsNode();
+	if (!SelectedNode.IsValid() || SelectedNode->IsGroup())
 	{
 		return;
 	}
 
-	const uint32 CounterId = SelectedNodes[0]->GetCounterId();
-	FString CounterName = SelectedNodes[0]->GetName().ToString();
+	const uint32 CounterId = SelectedNode->GetCounterId();
+	FString CounterName = SelectedNode->GetName().ToString();
 
 	const FString DialogTitle = LOCTEXT("ExportOps_Title", "Export Counter Ops").ToString();
 
@@ -2700,8 +2585,7 @@ void SStatsView::ContextMenu_ExportOps_Execute() const
 	Params.IntervalStartTime = -std::numeric_limits<double>::infinity();
 	Params.IntervalEndTime = +std::numeric_limits<double>::infinity();
 
-	TSharedPtr<STimingProfilerWindow> Wnd = FTimingProfilerManager::Get()->GetProfilerWindow();
-	TSharedPtr<STimingView> TimingView = Wnd.IsValid() ? Wnd->GetTimingView() : nullptr;
+	TSharedPtr<STimingView> TimingView = GetTimingView();
 	if (TimingView.IsValid())
 	{
 		const double SelectionStartTime = TimingView->GetSelectionStartTime();
@@ -2749,8 +2633,6 @@ void SStatsView::ContextMenu_ExportCounters_Execute() const
 
 ////////////////////////////////////////////////////////////////////////////////////////////////////
 
-////////////////////////////////////////////////////////////////////////////////////////////////////
-
 bool SStatsView::OpenSaveTextFileDialog(const FString& InDialogTitle, const FString& InDefaultFile, FString& OutFilename) const
 {
 	TArray<FString> SaveFilenames;
@@ -2798,26 +2680,159 @@ IFileHandle* SStatsView::OpenExportFile(const TCHAR* InFilename) const
 	return ExportFileHandle;
 }
 
+////////////////////////////////////////////////////////////////////////////////////////////////////
+// Generate and Edit Colors
 ////////////////////////////////////////////////////////////////////////////////////////////////////
 
-FReply SStatsView::OnKeyDown(const FGeometry& MyGeometry, const FKeyEvent& InKeyEvent)
+bool SStatsView::ContextMenu_GenerateNewColor_CanExecute() const
 {
-	return CommandList->ProcessCommandBindings(InKeyEvent) == true ? FReply::Handled() : FReply::Unhandled();
+	return TreeView->GetNumItemsSelected() > 0;
 }
 
 ////////////////////////////////////////////////////////////////////////////////////////////////////
 
-bool SStatsView::AreSelectedCounterNodesValid() const
+void SStatsView::ContextMenu_GenerateNewColor_Execute() const
 {
-	const TArray<FStatsNodePtr> SelectedNodes = TreeView->GetSelectedItems();
+	TArray<FStatsNodePtr> SelectedNodes = TreeView->GetSelectedItems();
 	for (const FStatsNodePtr& SelectedNode : SelectedNodes)
 	{
-		if (!SelectedNode.IsValid())
+		if (SelectedNode->Is<FStatsNode>())
 		{
-			return false;
+			FStatsNode& StatsNode = SelectedNode->As<FStatsNode>();
+			SetRandomColorToNode(StatsNode);
 		}
 	}
-	return true;
+}
+
+////////////////////////////////////////////////////////////////////////////////////////////////////
+
+void SStatsView::SetColorToNode(FStatsNode& StatsNode, FLinearColor CustomColor) const
+{
+	StatsNode.SetColor(CustomColor);
+
+	const uint32 CounterId = StatsNode.GetCounterId();
+	const FLinearColor BorderColor(CustomColor.R + 0.4f, CustomColor.G + 0.4f, CustomColor.B + 0.4f, 1.0f);
+
+	TSharedPtr<FTimingGraphTrack> GraphTrack = GetTimingViewMainGraphTrack();
+	if (GraphTrack.IsValid())
+	{
+		TSharedPtr<FTimingGraphSeries> CounterSeries = GraphTrack->GetStatsCounterSeries(CounterId);
+		if (CounterSeries.IsValid())
+		{
+			CounterSeries->SetColor(CustomColor, BorderColor);
+		}
+	}
+}
+
+////////////////////////////////////////////////////////////////////////////////////////////////////
+
+void SStatsView::SetRandomColorToNode(FStatsNode& StatsNode) const
+{
+	uint64 Time = FPlatformTime::Cycles64();
+	uint32 Hash = (Time & 0xFFFFFFFF) ^ (Time >> 32);
+	Hash = ((Hash << 5) + Hash) * 0xfb23618f;
+
+	const uint8 H = Hash & 0xFF;
+	const uint8 S = 128 + ((Hash >> 8) & 0x7F);
+	const uint8 V = 128 + ((Hash >> 16) & 0x7F);
+	FLinearColor CustomColor = FLinearColor::MakeFromHSV8(H, S, V);
+	CustomColor.A = 1.0f;
+
+	SetColorToNode(StatsNode, CustomColor);
+}
+
+////////////////////////////////////////////////////////////////////////////////////////////////////
+
+bool SStatsView::ContextMenu_EditColor_CanExecute() const
+{
+	return TreeView->GetNumItemsSelected() > 0;
+}
+
+////////////////////////////////////////////////////////////////////////////////////////////////////
+
+void SStatsView::ContextMenu_EditColor_Execute()
+{
+	EditableColorValue = FLinearColor(0.5f, 0.5f, 0.5f, 1.0f);
+
+	TArray<FStatsNodePtr> SelectedNodes = TreeView->GetSelectedItems();
+	for (const FStatsNodePtr& SelectedNode : SelectedNodes)
+	{
+		if (SelectedNode->Is<FStatsNode>())
+		{
+			FStatsNode& StatsNode = SelectedNode->As<FStatsNode>();
+			EditableColorValue = StatsNode.GetColor();
+			break;
+		}
+	}
+
+	FColorPickerArgs PickerArgs;
+	{
+		PickerArgs.bUseAlpha = true;
+		PickerArgs.bOnlyRefreshOnMouseUp = false;
+		PickerArgs.bOnlyRefreshOnOk = false;
+		PickerArgs.bExpandAdvancedSection = false;
+		PickerArgs.OnColorCommitted = FOnLinearColorValueChanged::CreateSP(this, &SStatsView::SetEditableColor);
+		PickerArgs.OnColorPickerCancelled = FOnColorPickerCancelled::CreateSP(this, &SStatsView::ColorPickerCancelled);
+		PickerArgs.InitialColor = EditableColorValue;
+		PickerArgs.ParentWidget = SharedThis(this);
+	}
+
+	OpenColorPicker(PickerArgs);
+}
+
+////////////////////////////////////////////////////////////////////////////////////////////////////
+
+FLinearColor SStatsView::GetEditableColor() const
+{
+	return EditableColorValue;
+}
+
+////////////////////////////////////////////////////////////////////////////////////////////////////
+
+void SStatsView::SetEditableColor(FLinearColor NewColor)
+{
+	EditableColorValue = NewColor;
+
+	TArray<FStatsNodePtr> SelectedNodes = TreeView->GetSelectedItems();
+	for (const FStatsNodePtr& SelectedNode : SelectedNodes)
+	{
+		if (SelectedNode->Is<FStatsNode>())
+		{
+			FStatsNode& StatsNode = SelectedNode->As<FStatsNode>();
+			SetColorToNode(StatsNode, EditableColorValue);
+		}
+	}
+}
+
+////////////////////////////////////////////////////////////////////////////////////////////////////
+
+void SStatsView::ColorPickerCancelled(FLinearColor OriginalColor)
+{
+	SetEditableColor(OriginalColor);
+}
+
+////////////////////////////////////////////////////////////////////////////////////////////////////
+// Misc
+////////////////////////////////////////////////////////////////////////////////////////////////////
+
+TSharedPtr<STimingProfilerWindow> SStatsView::GetProfilerWindow() const
+{
+	return FTimingProfilerManager::Get()->GetProfilerWindow();
+}
+
+////////////////////////////////////////////////////////////////////////////////////////////////////
+
+TSharedPtr<STimingView> SStatsView::GetTimingView() const
+{
+	TSharedPtr<STimingProfilerWindow> ProfilerWindow = GetProfilerWindow();
+	return ProfilerWindow.IsValid() ? ProfilerWindow->GetTimingView() : nullptr;
+}
+
+////////////////////////////////////////////////////////////////////////////////////////////////////
+
+FReply SStatsView::OnKeyDown(const FGeometry& MyGeometry, const FKeyEvent& InKeyEvent)
+{
+	return CommandList->ProcessCommandBindings(InKeyEvent) == true ? FReply::Handled() : FReply::Unhandled();
 }
 
 ////////////////////////////////////////////////////////////////////////////////////////////////////
diff --git a/Engine/Source/Developer/TraceInsights/Private/Insights/TimingProfiler/Widgets/SStatsView.h b/Engine/Source/Developer/TraceInsights/Private/Insights/TimingProfiler/Widgets/SStatsView.h
--- a/Engine/Source/Developer/TraceInsights/Private/Insights/TimingProfiler/Widgets/SStatsView.h
+++ b/Engine/Source/Developer/TraceInsights/Private/Insights/TimingProfiler/Widgets/SStatsView.h
@@ -43,6 +43,8 @@ namespace UE::Insights::TimingProfiler
 
 class FCounterAggregator;
 class FTimingGraphTrack;
+class STimingProfilerWindow;
+class STimingView;
 
 ////////////////////////////////////////////////////////////////////////////////////////////////////
 
@@ -123,22 +125,8 @@ private:
 	void TreeView_BuildSortByMenu(FMenuBuilder& MenuBuilder);
 	void TreeView_BuildViewColumnMenu(FMenuBuilder& MenuBuilder);
 
-	//////////////////////////////////////////////////
-	// Generate and Edit Counters Colors
-
-	bool ContextMenu_GenerateNewColor_CanExecute() const;
-	void ContextMenu_GenerateNewColor_Execute() const;
-
-	bool ContextMenu_EditColor_CanExecute() const;
-	void ContextMenu_EditColor_Execute();
-
-	void SetColorToNode(FStatsNode& StatsNode, FLinearColor CustomColor) const;
-	void SetRandomColorToNode(FStatsNode& StatsNode) const;
-	FLinearColor GetEditableColor() const;
-	void SetEditableColor(FLinearColor NewColor);
-	void ColorPickerCancelled(FLinearColor OriginalColor);
-
-	//////////////////////////////////////////////////
+	////////////////////////////////////////////////////////////////////////////////////////////////////
+	// Copy to Clipboard
 
 	bool ContextMenu_CopySelectedToClipboard_CanExecute() const;
 	void ContextMenu_CopySelectedToClipboard_Execute();
@@ -146,6 +134,9 @@ private:
 	bool ContextMenu_CopySelectedNameToClipboard_CanExecute() const;
 	void ContextMenu_CopySelectedNameToClipboard_Execute();
 
+	////////////////////////////////////////////////////////////////////////////////////////////////////
+	// Export
+
 	bool ContextMenu_Export_CanExecute() const;
 	void ContextMenu_Export_Execute();
 
@@ -161,6 +152,22 @@ private:
 	bool OpenSaveTextFileDialog(const FString& InDialogTitle, const FString& InDefaultFile, FString& OutFilename) const;
 	class IFileHandle* OpenExportFile(const TCHAR* InFilename) const;
 
+	////////////////////////////////////////////////////////////////////////////////////////////////////
+	// Generate and Edit Counters Colors
+
+	bool ContextMenu_GenerateNewColor_CanExecute() const;
+	void ContextMenu_GenerateNewColor_Execute() const;
+
+	bool ContextMenu_EditColor_CanExecute() const;
+	void ContextMenu_EditColor_Execute();
+
+	void SetColorToNode(FStatsNode& StatsNode, FLinearColor CustomColor) const;
+	void SetRandomColorToNode(FStatsNode& StatsNode) const;
+
+	FLinearColor GetEditableColor() const;
+	void SetEditableColor(FLinearColor NewColor);
+	void ColorPickerCancelled(FLinearColor OriginalColor);
+
 	////////////////////////////////////////////////////////////////////////////////////////////////////
 	// Tree View - Columns' Header
 
@@ -185,12 +192,11 @@ private:
 	/** Called by STreeView when selection has changed. */
 	void TreeView_OnSelectionChanged(FStatsNodePtr SelectedItem, ESelectInfo::Type SelectInfo);
 
+	FStatsNodePtr GetSingleSelectedStatsNode() const;
+
 	/** Called by STreeView when a tree item is double clicked. */
 	void TreeView_OnMouseButtonDoubleClick(FStatsNodePtr TreeNode);
 
-	/** Checks validity of one or more selected nodes*/
-	bool AreSelectedCounterNodesValid() const;
-
 	////////////////////////////////////////////////////////////////////////////////////////////////////
 	// Tree View - Table Row
 
@@ -256,7 +262,7 @@ private:
 	void OnSortModeChanged(const EColumnSortPriority::Type SortPriority, const FName& ColumnId, const EColumnSortMode::Type SortMode);
 
 	////////////////////////////////////////////////////////////////////////////////////////////////////
-	// Sorting actions
+	// Sorting Actions
 
 	// SortMode (HeaderMenu)
 	bool HeaderMenu_SortMode_IsChecked(const FName ColumnId, const EColumnSortMode::Type InSortMode);
@@ -274,7 +280,7 @@ private:
 	void ContextMenu_SortByColumn_Execute(const FName ColumnId);
 
 	////////////////////////////////////////////////////////////////////////////////////////////////////
-	// Column visibility actions
+	// Column Visibility Actions
 
 	// ShowColumn
 	bool CanShowColumn(const FName ColumnId) const;
@@ -302,8 +308,10 @@ private:
 	void ContextMenu_ResetColumns_Execute();
 
 	////////////////////////////////////////////////////////////////////////////////////////////////////
+	// Plot (Graph Series)
 
 	TSharedPtr<FTimingGraphTrack> GetTimingViewMainGraphTrack() const;
+
 	void ToggleGraphSeries(TSharedRef<FTimingGraphTrack> GraphTrack, FStatsNodeRef NodePtr) const;
 	bool IsSeriesInTimingViewMainGraph(FStatsNodePtr CounterNode) const;
 	void ToggleTimingViewMainGraphEventSeries(FStatsNodePtr CounterNode) const;
@@ -319,6 +327,9 @@ private:
 	 */
 	virtual void Tick(const FGeometry& AllottedGeometry, const double InCurrentTime, const float InDeltaTime) override;
 
+	TSharedPtr<STimingProfilerWindow> GetProfilerWindow() const;
+	TSharedPtr<STimingView> GetTimingView() const;
+
 private:
 	/** Table view model. */
 	TSharedPtr<FTable> Table;
@@ -374,7 +385,7 @@ private:
 	TSet<FStatsNodePtr> ExpandedNodes;
 
 	/** If true, the expanded nodes have been saved before applying a text filter. */
-	bool bExpansionSaved;
+	bool bExpansionSaved = false;
 
 	//////////////////////////////////////////////////
 	// Search box and filters
@@ -395,7 +406,7 @@ private:
 	bool FilterByDataType[static_cast<int>(EStatsNodeDataType::InvalidOrMax)];
 
 	/** Filter out the counters having zero total instance count (aggregated stats). */
-	bool bFilterOutZeroCountStats;
+	bool bFilterOutZeroCountStats = false;
 
 	//////////////////////////////////////////////////
 	// Grouping
@@ -405,7 +416,7 @@ private:
 	TSharedPtr<SComboBox<TSharedPtr<EStatsGroupingMode>>> GroupByComboBox;
 
 	/** How we group the counters? */
-	EStatsGroupingMode GroupingMode;
+	EStatsGroupingMode GroupingMode = EStatsGroupingMode::Flat;
 
 	//////////////////////////////////////////////////
 	// Sorting
@@ -420,7 +431,7 @@ private:
 	FName ColumnBeingSorted;
 
 	/** How we sort the nodes? Ascending or Descending. */
-	EColumnSortMode::Type ColumnSortMode;
+	EColumnSortMode::Type ColumnSortMode = EColumnSortMode::Type::Ascending;
 
 	//////////////////////////////////////////////////
 
diff --git a/Engine/Source/Developer/TraceInsights/Private/Insights/TimingProfiler/Widgets/STimersView.cpp b/Engine/Source/Developer/TraceInsights/Private/Insights/TimingProfiler/Widgets/STimersView.cpp
--- a/Engine/Source/Developer/TraceInsights/Private/Insights/TimingProfiler/Widgets/STimersView.cpp
+++ b/Engine/Source/Developer/TraceInsights/Private/Insights/TimingProfiler/Widgets/STimersView.cpp
@@ -28,9 +28,9 @@
 
 // TraceServices
 #include "Common/ProviderLock.h"
+#include "TraceServices/Model/StackSamples.h"
 #include "TraceServices/Model/Threads.h"
 #include "TraceServices/Model/TimingProfiler.h"
-#include "TraceServices/Model/StackSamples.h"
 
 // TraceInsightsCore
 #include "InsightsCore/Common/Stopwatch.h"
@@ -603,9 +603,8 @@ END_SLATE_FUNCTION_BUILD_OPTIMIZATION
 
 TSharedPtr<SWidget> STimersView::TreeView_GetMenuContent()
 {
-	const TArray<FTimerNodePtr> SelectedNodes = TreeView->GetSelectedItems();
-	const int32 NumSelectedNodes = SelectedNodes.Num();
-	FTimerNodePtr SelectedNode = NumSelectedNodes ? SelectedNodes[0] : nullptr;
+	const int32 NumSelectedNodes = TreeView->GetNumItemsSelected();
+	FTimerNodePtr SingleSelectedNode;
 
 	FText SelectionStr;
 	if (NumSelectedNodes == 0)
@@ -614,7 +613,11 @@ TSharedPtr<SWidget> STimersView::TreeView_GetMenuContent()
 	}
 	else if (NumSelectedNodes == 1)
 	{
-		FString ItemName = SelectedNode->GetName().ToString();
+		TArray<FTimerNodePtr> SelectedNodes = TreeView->GetSelectedItems();
+		check(SelectedNodes.Num() == 1);
+		SingleSelectedNode = SelectedNodes[0];
+
+		FString ItemName = SingleSelectedNode->GetName().ToString();
 		const int32 MaxStringLen = 64;
 		if (ItemName.Len() > MaxStringLen)
 		{
@@ -658,23 +661,26 @@ TSharedPtr<SWidget> STimersView::TreeView_GetMenuContent()
 	// Timer options section
 	MenuBuilder.BeginSection("TimerOptions", LOCTEXT("ContextMenu_Section_TimerOptions", "Timer Options"));
 	{
-		TSharedPtr<STimingView> TimingView = GetTimingView();
+		const bool bIsValidNode =
+			SingleSelectedNode.IsValid() &&
+			SingleSelectedNode->GetType() != ETimerNodeType::Group;
 
-		auto CanExecute = [TimingView, NumSelectedNodes, SelectedNode]()
+		auto CanExecute = [this, bIsValidNode]()
 		{
-			return TimingView.IsValid() && NumSelectedNodes == 1 && SelectedNode.IsValid() && SelectedNode->GetType() != ETimerNodeType::Group;
+			return bIsValidNode && GetTimingView().IsValid();
 		};
 
 		// Highlight event
 		{
 			FUIAction Action_ToggleHighlight;
 			Action_ToggleHighlight.CanExecuteAction = FCanExecuteAction::CreateLambda(CanExecute);
-			Action_ToggleHighlight.ExecuteAction = FExecuteAction::CreateSP(this, &STimersView::ToggleTimingViewEventFilter, SelectedNode);
+			Action_ToggleHighlight.ExecuteAction = FExecuteAction::CreateSP(this, &STimersView::ToggleTimingViewEventFilter, SingleSelectedNode);
 
-			if (SelectedNode.IsValid() &&
-				SelectedNode->GetType() != ETimerNodeType::Group &&
+			TSharedPtr<STimingView> TimingView = GetTimingView();
+
+			if (bIsValidNode &&
 				TimingView.IsValid() &&
-				TimingView->IsFilterByEventType(SelectedNode->GetTimerId()))
+				TimingView->IsFilterByEventType(SingleSelectedNode->GetTimerId()))
 			{
 				MenuBuilder.AddMenuEntry
 				(
@@ -726,9 +732,9 @@ TSharedPtr<SWidget> STimersView::TreeView_GetMenuContent()
 			FString File;
 			uint32 Line = 0;
 			bool bIsValidSource = false;
-			if (NumSelectedNodes == 1 && SelectedNode.IsValid() && SelectedNode->GetType() != ETimerNodeType::Group)
+			if (SingleSelectedNode.IsValid() && SingleSelectedNode->GetType() != ETimerNodeType::Group)
 			{
-				bIsValidSource = SelectedNode->GetSourceFileAndLine(File, Line);
+				bIsValidSource = SingleSelectedNode->GetSourceFileAndLine(File, Line);
 			}
 			FString UsablePath;
 			const bool bPathExists = FInsightsManager::Get()->GetSourceFilePathHelper().GetUsableFilePath(File, UsablePath);
@@ -999,22 +1005,17 @@ void STimersView::TreeView_BuildPlotTimerMenu(FMenuBuilder& MenuBuilder)
 {
 	MenuBuilder.SetSearchable(false);
 
-	const TArray<FTimerNodePtr> SelectedNodes = TreeView->GetSelectedItems();
-	const int32 NumSelectedNodes = SelectedNodes.Num();
-	FTimerNodePtr SelectedNode = NumSelectedNodes ? SelectedNodes[0] : nullptr;
+	FTimerNodePtr SelectedNode = GetSingleSelectedTimerNode();
+	const bool bIsValidNode = SelectedNode.IsValid() && SelectedNode->GetType() != ETimerNodeType::Group;
 
-	auto CanExecuteAddToGraphTrack = [NumSelectedNodes, SelectedNode]()
+	auto CanExecuteAddToGraphTrack = [this, bIsValidNode]()
 	{
-		TSharedPtr<STimingProfilerWindow> Wnd = FTimingProfilerManager::Get()->GetProfilerWindow();
-		TSharedPtr<STimingView> TimingView = Wnd.IsValid() ? Wnd->GetTimingView() : nullptr;
-		return TimingView.IsValid() && NumSelectedNodes == 1 && SelectedNode.IsValid() && SelectedNode->GetType() != ETimerNodeType::Group;
+		return bIsValidNode && GetTimingView().IsValid();
 	};
 
-	auto CanExecuteAddToFramesTrack = [NumSelectedNodes, SelectedNode]()
+	auto CanExecuteAddToFramesTrack = [this, bIsValidNode]()
 	{
-		TSharedPtr<STimingProfilerWindow> Wnd = FTimingProfilerManager::Get()->GetProfilerWindow();
-		TSharedPtr<SFrameTrack> FrameTrack = Wnd.IsValid() ? Wnd->GetFrameView() : nullptr;
-		return FrameTrack.IsValid() && NumSelectedNodes == 1 && SelectedNode.IsValid() && SelectedNode->GetType() != ETimerNodeType::Group;
+		return bIsValidNode && GetFrameTrack().IsValid();
 	};
 
 	MenuBuilder.BeginSection("Instance", LOCTEXT("Plot_Series_Instance_Section", "Instance"));
@@ -1025,8 +1026,7 @@ void STimersView::TreeView_BuildPlotTimerMenu(FMenuBuilder& MenuBuilder)
 		Action_ToggleTimerInGraphTrack.CanExecuteAction = FCanExecuteAction::CreateLambda(CanExecuteAddToGraphTrack);
 		Action_ToggleTimerInGraphTrack.ExecuteAction = FExecuteAction::CreateSP(this, &STimersView::ToggleTimingViewMainGraphEventInstanceSeries, SelectedNode);
 
-		if (SelectedNode.IsValid() &&
-			SelectedNode->GetType() != ETimerNodeType::Group &&
+		if (bIsValidNode &&
 			IsInstanceSeriesInTimingViewMainGraph(SelectedNode))
 		{
 			MenuBuilder.AddMenuEntry
@@ -1063,8 +1063,7 @@ void STimersView::TreeView_BuildPlotTimerMenu(FMenuBuilder& MenuBuilder)
 		Action_ToggleFrameStatsTimerInGraphTrack.CanExecuteAction = FCanExecuteAction::CreateLambda(CanExecuteAddToGraphTrack);
 		Action_ToggleFrameStatsTimerInGraphTrack.ExecuteAction = FExecuteAction::CreateSP(this, &STimersView::ToggleTimingViewMainGraphEventFrameStatsSeries, SelectedNode, ETraceFrameType::TraceFrameType_Game);
 
-		if (SelectedNode.IsValid() &&
-			SelectedNode->GetType() != ETimerNodeType::Group &&
+		if (bIsValidNode &&
 			IsFrameStatsSeriesInTimingViewMainGraph(SelectedNode, ETraceFrameType::TraceFrameType_Game))
 		{
 			MenuBuilder.AddMenuEntry
@@ -1097,8 +1096,7 @@ void STimersView::TreeView_BuildPlotTimerMenu(FMenuBuilder& MenuBuilder)
 		Action_ToggleFrameStatsTimerInFrameTrack.CanExecuteAction = FCanExecuteAction::CreateLambda(CanExecuteAddToFramesTrack);
 		Action_ToggleFrameStatsTimerInFrameTrack.ExecuteAction = FExecuteAction::CreateSP(this, &STimersView::ToggleFrameTrackSeries, SelectedNode, ETraceFrameType::TraceFrameType_Game);
 
-		if (SelectedNode.IsValid() &&
-			SelectedNode->GetType() != ETimerNodeType::Group &&
+		if (bIsValidNode &&
 			IsSeriesInFrameTrack(SelectedNode, ETraceFrameType::TraceFrameType_Game))
 		{
 			MenuBuilder.AddMenuEntry
@@ -1135,8 +1133,7 @@ void STimersView::TreeView_BuildPlotTimerMenu(FMenuBuilder& MenuBuilder)
 		Action_ToggleFrameStatsTimerInGraphTrack.CanExecuteAction = FCanExecuteAction::CreateLambda(CanExecuteAddToGraphTrack);
 		Action_ToggleFrameStatsTimerInGraphTrack.ExecuteAction = FExecuteAction::CreateSP(this, &STimersView::ToggleTimingViewMainGraphEventFrameStatsSeries, SelectedNode, ETraceFrameType::TraceFrameType_Rendering);
 
-		if (SelectedNode.IsValid() &&
-			SelectedNode->GetType() != ETimerNodeType::Group &&
+		if (bIsValidNode &&
 			IsFrameStatsSeriesInTimingViewMainGraph(SelectedNode, ETraceFrameType::TraceFrameType_Rendering))
 		{
 			MenuBuilder.AddMenuEntry
@@ -1169,8 +1166,7 @@ void STimersView::TreeView_BuildPlotTimerMenu(FMenuBuilder& MenuBuilder)
 		Action_ToggleFrameStatsTimerInFrameTrack.CanExecuteAction = FCanExecuteAction::CreateLambda(CanExecuteAddToFramesTrack);
 		Action_ToggleFrameStatsTimerInFrameTrack.ExecuteAction = FExecuteAction::CreateSP(this, &STimersView::ToggleFrameTrackSeries, SelectedNode, ETraceFrameType::TraceFrameType_Rendering);
 
-		if (SelectedNode.IsValid() &&
-			SelectedNode->GetType() != ETimerNodeType::Group &&
+		if (bIsValidNode &&
 			IsSeriesInFrameTrack(SelectedNode, ETraceFrameType::TraceFrameType_Rendering))
 		{
 			MenuBuilder.AddMenuEntry
@@ -1623,9 +1619,11 @@ void STimersView::ApplyFiltering()
 			check(ChildPtr->Is<FTimerNode>());
 			const FTimerNodePtr& NodePtr = StaticCastSharedPtr<FTimerNode>(ChildPtr);
 
-			const bool bIsChildVisible = (!bFilterOutZeroCountTimers || NodePtr->GetAggregatedStats().InstanceCount > 0)
-									  && bTimerTypeIsVisible[static_cast<int>(NodePtr->GetType())]
-									  && Filters->PassesAllFilters(NodePtr);
+			const bool bIsChildVisible =
+				(!bFilterOutZeroCountTimers || NodePtr->GetAggregatedStats().InstanceCount > 0) &&
+				bTimerTypeIsVisible[static_cast<int>(NodePtr->GetType())] &&
+				Filters->PassesAllFilters(NodePtr);
+
 			if (bIsChildVisible)
 			{
 				// Add a child.
@@ -1783,17 +1781,17 @@ ECheckBoxState STimersView::FilterOutZeroCountTimers_IsChecked() const
 
 ////////////////////////////////////////////////////////////////////////////////////////////////////
 
-void STimersView::FilterByTimerType_OnCheckStateChanged(ECheckBoxState NewRadioState, const ETimerNodeType InStatType)
+void STimersView::FilterByTimerType_OnCheckStateChanged(ECheckBoxState NewRadioState, const ETimerNodeType InNodeType)
 {
-	bTimerTypeIsVisible[static_cast<int>(InStatType)] = (NewRadioState == ECheckBoxState::Checked);
+	bTimerTypeIsVisible[static_cast<int>(InNodeType)] = (NewRadioState == ECheckBoxState::Checked);
 	ApplyFiltering();
 }
 
 ////////////////////////////////////////////////////////////////////////////////////////////////////
 
-ECheckBoxState STimersView::FilterByTimerType_IsChecked(const ETimerNodeType InStatType) const
+ECheckBoxState STimersView::FilterByTimerType_IsChecked(const ETimerNodeType InNodeType) const
 {
-	return bTimerTypeIsVisible[static_cast<int>(InStatType)] ? ECheckBoxState::Checked : ECheckBoxState::Unchecked;
+	return bTimerTypeIsVisible[static_cast<int>(InNodeType)] ? ECheckBoxState::Checked : ECheckBoxState::Unchecked;
 }
 
 ////////////////////////////////////////////////////////////////////////////////////////////////////
@@ -1830,12 +1828,9 @@ FTimerNodePtr STimersView::GetSingleSelectedTimerNode() const
 	{
 		return nullptr;
 	}
-	const TArray<FTimerNodePtr> SelectedNodes = TreeView->GetSelectedItems();
-	if (SelectedNodes.Num() == 1)
-	{
-		return SelectedNodes[0];
-	}
-	return nullptr;
+	TArray<FTimerNodePtr> SelectedNodes = TreeView->GetSelectedItems();
+	check(SelectedNodes.Num() == 1);
+	return SelectedNodes[0];
 }
 
 ////////////////////////////////////////////////////////////////////////////////////////////////////
@@ -2183,7 +2178,9 @@ void STimersView::CreateGroupByOptionsSources()
 
 void STimersView::GroupBy_OnSelectionChanged(TSharedPtr<ETimerGroupingMode> NewGroupingMode, ESelectInfo::Type SelectInfo)
 {
-	if (SelectInfo != ESelectInfo::Direct && *NewGroupingMode != GroupingMode)
+	if (SelectInfo != ESelectInfo::Direct &&
+		NewGroupingMode.IsValid() &&
+		*NewGroupingMode != GroupingMode)
 	{
 		GroupingMode = *NewGroupingMode;
 
@@ -2197,6 +2194,7 @@ void STimersView::GroupBy_OnSelectionChanged(TSharedPtr<ETimerGroupingMode> NewG
 
 TSharedRef<SWidget> STimersView::GroupBy_OnGenerateWidget(TSharedPtr<ETimerGroupingMode> InGroupingMode) const
 {
+	check(InGroupingMode.IsValid());
 	return SNew(STextBlock)
 		.Text(TimerNodeGroupingHelper::ToText(*InGroupingMode))
 		.ToolTipText(TimerNodeGroupingHelper::ToDescription(*InGroupingMode));
@@ -2216,7 +2214,6 @@ FText STimersView::GroupBy_GetSelectedTooltipText() const
 	return TimerNodeGroupingHelper::ToDescription(GroupingMode);
 }
 
-
 ////////////////////////////////////////////////////////////////////////////////////////////////////
 // Frame Type
 ////////////////////////////////////////////////////////////////////////////////////////////////////
@@ -2242,15 +2239,15 @@ void STimersView::CreateModeOptionsSources()
 
 ////////////////////////////////////////////////////////////////////////////////////////////////////
 
-void STimersView::Mode_OnSelectionChanged(TSharedPtr<ETraceFrameType> NewFrameType, ESelectInfo::Type SelectInfo)
+void STimersView::Mode_OnSelectionChanged(TSharedPtr<ETraceFrameType> InModeFrameType, ESelectInfo::Type SelectInfo)
 {
-	if (*NewFrameType == ModeFrameType)
+	if (!InModeFrameType.IsValid() || *InModeFrameType == ModeFrameType)
 	{
 		return;
 	}
 
 	SaveVisibleColumnsSettings();
-	ModeFrameType = *NewFrameType;
+	ModeFrameType = *InModeFrameType;
 	LoadVisibleColumnsSettings();
 
 	SetTimingViewFrameType();
@@ -2262,11 +2259,12 @@ void STimersView::Mode_OnSelectionChanged(TSharedPtr<ETraceFrameType> NewFrameTy
 
 ////////////////////////////////////////////////////////////////////////////////////////////////////
 
-TSharedRef<SWidget> STimersView::Mode_OnGenerateWidget(TSharedPtr<ETraceFrameType> InFrameTypeMode) const
+TSharedRef<SWidget> STimersView::Mode_OnGenerateWidget(TSharedPtr<ETraceFrameType> InModeFrameType) const
 {
+	check(InModeFrameType.IsValid());
 	return SNew(STextBlock)
-	.Text(Mode_GetText(*InFrameTypeMode))
-	.ToolTipText(Mode_GetTooltipText(*InFrameTypeMode));
+		.Text(Mode_GetText(*InModeFrameType))
+		.ToolTipText(Mode_GetTooltipText(*InModeFrameType));
 }
 
 ////////////////////////////////////////////////////////////////////////////////////////////////////
@@ -2841,6 +2839,7 @@ void STimersView::RebuildTree(bool bResync)
 			if (TimerCount != PreviousNodeCount)
 			{
 				check(TimerCount > PreviousNodeCount);
+
 				TimerNodes.Reserve(TimerCount);
 
 				TSharedPtr<SFrameTrack> FrameTrack = GetFrameTrack();
@@ -2882,7 +2881,7 @@ void STimersView::RebuildTree(bool bResync)
 					TimerNodePtr->SetAddedToGraphsFlag(bIsTimerAddedToGraphs);
 
 					FLinearColor ColorOverride;
-					TSharedPtr<STimingProfilerWindow> ProfilerWindow = FTimingProfilerManager::Get()->GetProfilerWindow();
+					TSharedPtr<STimingProfilerWindow> ProfilerWindow = GetProfilerWindow();
 					if (ProfilerWindow.IsValid() && ProfilerWindow->GetTimerColor(Timer.Id, ColorOverride))
 					{
 						TimerNodePtr->SetColor(ColorOverride);
@@ -2957,7 +2956,8 @@ void STimersView::ResetStats()
 
 void STimersView::UpdateStats(double StartTime, double EndTime)
 {
-	if (Aggregator->GetIntervalStartTime() != StartTime || Aggregator->GetIntervalEndTime() != EndTime)
+	if (Aggregator->GetIntervalStartTime() != StartTime ||
+		Aggregator->GetIntervalEndTime() != EndTime)
 	{
 		Aggregator->Cancel();
 		Aggregator->SetTimeInterval(StartTime, EndTime);
@@ -2991,7 +2991,7 @@ void STimersView::FinishAggregation()
 	UpdateTree(); // grouping + sorting + filtering
 
 	// Ensure the last selected item is visible.
-	const TArray<FTimerNodePtr> SelectedNodes = TreeView->GetSelectedItems();
+	TArray<FTimerNodePtr> SelectedNodes = TreeView->GetSelectedItems();
 	if (SelectedNodes.Num() > 0)
 	{
 		TreeView->RequestScrollIntoView(SelectedNodes.Last());
@@ -3068,7 +3068,7 @@ TSharedPtr<FTimingGraphTrack> STimersView::GetTimingViewMainGraphTrack() const
 
 TSharedPtr<SFrameTrack> STimersView::GetFrameTrack() const
 {
-	TSharedPtr<STimingProfilerWindow> Wnd = FTimingProfilerManager::Get()->GetProfilerWindow();
+	TSharedPtr<STimingProfilerWindow> Wnd = GetProfilerWindow();
 	return Wnd.IsValid() ? Wnd->GetFrameView() : nullptr;
 }
 
@@ -3117,7 +3117,7 @@ void STimersView::ToggleGraphInstanceSeries(TSharedRef<FTimingGraphTrack> GraphT
 		GraphTrack->SetDirtyFlag();
 	}
 
-	TSharedPtr<STimingProfilerWindow> ProfilerWindow = FTimingProfilerManager::Get()->GetProfilerWindow();
+	TSharedPtr<STimingProfilerWindow> ProfilerWindow = GetProfilerWindow();
 	if (ProfilerWindow.IsValid())
 	{
 		ProfilerWindow->OnTimerAddedToGraphsChanged(TimerId);
@@ -3175,7 +3175,7 @@ void STimersView::ToggleGraphFrameStatsSeries(TSharedRef<FTimingGraphTrack> Grap
 		GraphTrack->SetDirtyFlag();
 	}
 
-	TSharedPtr<STimingProfilerWindow> ProfilerWindow = FTimingProfilerManager::Get()->GetProfilerWindow();
+	TSharedPtr<STimingProfilerWindow> ProfilerWindow = GetProfilerWindow();
 	if (ProfilerWindow.IsValid())
 	{
 		ProfilerWindow->OnTimerAddedToGraphsChanged(TimerId);
@@ -3267,7 +3267,7 @@ void STimersView::ToggleFrameTrackSeries(FTimerNodePtr TimerNode, ETraceFrameTyp
 		FrameTrack->AddTimerFrameStatSeries(FrameType, TimerId, TimerNode->GetColor(), SeriesName);
 	}
 
-	TSharedPtr<STimingProfilerWindow> ProfilerWindow = FTimingProfilerManager::Get()->GetProfilerWindow();
+	TSharedPtr<STimingProfilerWindow> ProfilerWindow = GetProfilerWindow();
 	if (ProfilerWindow.IsValid())
 	{
 		ProfilerWindow->OnTimerAddedToGraphsChanged(TimerId);
@@ -3280,9 +3280,7 @@ void STimersView::ToggleFrameTrackSeries(FTimerNodePtr TimerNode, ETraceFrameTyp
 
 bool STimersView::ContextMenu_CopyToClipboard_CanExecute() const
 {
-	const TArray<FTimerNodePtr> SelectedNodes = TreeView->GetSelectedItems();
-
-	return SelectedNodes.Num() > 0;
+	return TreeView->GetNumItemsSelected() > 0;
 }
 
 ////////////////////////////////////////////////////////////////////////////////////////////////////
@@ -3300,7 +3298,7 @@ void STimersView::ContextMenu_CopyToClipboard_Execute()
 
 bool STimersView::ContextMenu_CopyTimerNameToClipboard_CanExecute() const
 {
-	return TreeView->GetSelectedItems().Num() > 0;
+	return TreeView->GetNumItemsSelected() > 0;
 }
 
 ////////////////////////////////////////////////////////////////////////////////////////////////////
@@ -3320,8 +3318,7 @@ void STimersView::ContextMenu_CopyTimerNameToClipboard_Execute()
 
 bool STimersView::ContextMenu_Export_CanExecute() const
 {
-	const TArray<FTimerNodePtr> SelectedNodes = TreeView->GetSelectedItems();
-	return SelectedNodes.Num() > 0;
+	return TreeView->GetNumItemsSelected() > 0;
 }
 
 ////////////////////////////////////////////////////////////////////////////////////////////////////
@@ -3333,16 +3330,16 @@ void STimersView::ContextMenu_Export_Execute()
 		return;
 	}
 
-	const TArray<FTimerNodePtr> SelectedTimerNodes = TreeView->GetSelectedItems();
+	TArray<FTimerNodePtr> SelectedTimerNodes = TreeView->GetSelectedItems();
 	if (SelectedTimerNodes.Num() == 0)
 	{
 		return;
 	}
 
 	TArray<FBaseTreeNodePtr> SelectedNodes;
-	for (FTimerNodePtr TimerPtr : SelectedTimerNodes)
+	for (FTimerNodePtr Node : SelectedTimerNodes)
 	{
-		SelectedNodes.Add(TimerPtr);
+		SelectedNodes.Add(Node);
 	}
 
 	const FString DialogTitle = LOCTEXT("Export_Title", "Export Aggregated Timer Stats").ToString();
@@ -3463,9 +3460,7 @@ void STimersView::ContextMenu_Export_Execute()
 
 bool STimersView::ContextMenu_ExportTimingEventsSelection_CanExecute() const
 {
-	const TArray<FTimerNodePtr> SelectedNodes = TreeView->GetSelectedItems();
-
-	return SelectedNodes.Num() > 0;
+	return TreeView->GetNumItemsSelected() > 0;
 }
 
 ////////////////////////////////////////////////////////////////////////////////////////////////////
@@ -3844,14 +3839,14 @@ void STimersView::OpenSourceFileInIDE(FTimerNodePtr InNode) const
 
 bool STimersView::ContextMenu_GenerateNewColor_CanExecute() const
 {
-	return AreSelectedTimerNodesValid();
+	return TreeView->GetNumItemsSelected() > 0;
 }
 
 ////////////////////////////////////////////////////////////////////////////////////////////////////
 
 void STimersView::ContextMenu_GenerateNewColor_Execute() const
 {
-	const TArray<FTimerNodePtr> SelectedNodes = TreeView->GetSelectedItems();
+	TArray<FTimerNodePtr> SelectedNodes = TreeView->GetSelectedItems();
 	for (const FTimerNodePtr& SelectedNode : SelectedNodes)
 	{
 		if (SelectedNode->Is<FTimerNode>())
@@ -3860,20 +3855,18 @@ void STimersView::ContextMenu_GenerateNewColor_Execute() const
 			SetRandomColorToNode(TimerNode);
 		}
 	}
-};
+}
 
 ////////////////////////////////////////////////////////////////////////////////////////////////////
 
 void STimersView::SetColorToNode(FTimerNode& TimerNode, FLinearColor CustomColor) const
 {
 	TimerNode.SetColor(CustomColor);
-	if (FTimingProfilerManager::Get().IsValid())
+
+	TSharedPtr<STimingProfilerWindow> ProfilerWindow = GetProfilerWindow();
+	if (ProfilerWindow.IsValid())
 	{
-		TSharedPtr<STimingProfilerWindow> ProfilerWindow = FTimingProfilerManager::Get()->GetProfilerWindow();
-		if (ProfilerWindow.IsValid())
-		{
-			ProfilerWindow->SetTimerColor(TimerNode.GetTimerId(), CustomColor);
-		}
+		ProfilerWindow->SetTimerColor(TimerNode.GetTimerId(), CustomColor);
 	}
 }
 
@@ -3891,22 +3884,14 @@ void STimersView::SetRandomColorToNode(FTimerNode& TimerNode) const
 	FLinearColor CustomColor = FLinearColor::MakeFromHSV8(H, S, V);
 	CustomColor.A = 1.0f;
 
-	TimerNode.SetColor(CustomColor);
-	if (FTimingProfilerManager::Get().IsValid())
-	{
-		TSharedPtr<STimingProfilerWindow> ProfilerWindow = FTimingProfilerManager::Get()->GetProfilerWindow();
-		if (ProfilerWindow.IsValid())
-		{
-			ProfilerWindow->SetTimerColor(TimerNode.GetTimerId(), CustomColor);
-		}
-	}
+	SetColorToNode(TimerNode, CustomColor);
 }
 
 ////////////////////////////////////////////////////////////////////////////////////////////////////
 
 bool STimersView::ContextMenu_EditColor_CanExecute() const
 {
-	return AreSelectedTimerNodesValid();
+	return TreeView->GetNumItemsSelected() > 0;
 }
 
 ////////////////////////////////////////////////////////////////////////////////////////////////////
@@ -3914,7 +3899,8 @@ bool STimersView::ContextMenu_EditColor_CanExecute() const
 void STimersView::ContextMenu_EditColor_Execute()
 {
 	EditableColorValue = FLinearColor(0.5f, 0.5f, 0.5f, 1.0f);
-	const TArray<FTimerNodePtr> SelectedNodes = TreeView->GetSelectedItems();
+
+	TArray<FTimerNodePtr> SelectedNodes = TreeView->GetSelectedItems();
 	for (const FTimerNodePtr& SelectedNode : SelectedNodes)
 	{
 		if (SelectedNode->Is<FTimerNode>())
@@ -3953,7 +3939,7 @@ void STimersView::SetEditableColor(FLinearColor NewColor)
 {
 	EditableColorValue = NewColor;
 
-	const TArray<FTimerNodePtr> SelectedNodes = TreeView->GetSelectedItems();
+	TArray<FTimerNodePtr> SelectedNodes = TreeView->GetSelectedItems();
 	for (const FTimerNodePtr& SelectedNode : SelectedNodes)
 	{
 		if (SelectedNode->Is<FTimerNode>())
@@ -4152,6 +4138,7 @@ void STimersView::LoadSettings()
 }
 
 ////////////////////////////////////////////////////////////////////////////////////////////////////
+// Misc
 ////////////////////////////////////////////////////////////////////////////////////////////////////
 
 void STimersView::SetTimingViewFrameType()
@@ -4177,32 +4164,24 @@ void STimersView::OnTimingViewTrackListChanged()
 
 ////////////////////////////////////////////////////////////////////////////////////////////////////
 
-TSharedPtr<STimingView> STimersView::GetTimingView() const
+TSharedPtr<STimingProfilerWindow> STimersView::GetProfilerWindow() const
 {
-	TSharedPtr<STimingProfilerWindow> ProfilerWindow = FTimingProfilerManager::Get()->GetProfilerWindow();
-	return ProfilerWindow.IsValid() ? ProfilerWindow->GetTimingView() : nullptr;
+	return FTimingProfilerManager::Get()->GetProfilerWindow();
 }
 
 ////////////////////////////////////////////////////////////////////////////////////////////////////
 
-FReply STimersView::OnKeyDown(const FGeometry& MyGeometry, const FKeyEvent& InKeyEvent)
+TSharedPtr<STimingView> STimersView::GetTimingView() const
 {
-	return CommandList->ProcessCommandBindings(InKeyEvent) == true ? FReply::Handled() : FReply::Unhandled();
+	TSharedPtr<STimingProfilerWindow> ProfilerWindow = GetProfilerWindow();
+	return ProfilerWindow.IsValid() ? ProfilerWindow->GetTimingView() : nullptr;
 }
 
 ////////////////////////////////////////////////////////////////////////////////////////////////////
 
-bool STimersView::AreSelectedTimerNodesValid() const
+FReply STimersView::OnKeyDown(const FGeometry& MyGeometry, const FKeyEvent& InKeyEvent)
 {
-	const TArray<FTimerNodePtr> SelectedNodes = TreeView->GetSelectedItems();
-	for (const FTimerNodePtr& SelectedNode : SelectedNodes)
-	{
-		if (!SelectedNode.IsValid())
-		{
-			return false;
-		}
-	}
-	return true;
+	return CommandList->ProcessCommandBindings(InKeyEvent) == true ? FReply::Handled() : FReply::Unhandled();
 }
 
 ////////////////////////////////////////////////////////////////////////////////////////////////////
@@ -4246,6 +4225,7 @@ void STimersView::ToggleShowFormatInTimerNames_Execute()
 	Settings.SetAndSaveShowFormatInTimerNames(!Settings.GetShowFormatInTimerNames());
 
 	RebuildTree(true);
+
 	FTimingProfilerManager::Get()->UpdateCallersAndCallees();
 }
 
diff --git a/Engine/Source/Developer/TraceInsights/Private/Insights/TimingProfiler/Widgets/STimersView.h b/Engine/Source/Developer/TraceInsights/Private/Insights/TimingProfiler/Widgets/STimersView.h
--- a/Engine/Source/Developer/TraceInsights/Private/Insights/TimingProfiler/Widgets/STimersView.h
+++ b/Engine/Source/Developer/TraceInsights/Private/Insights/TimingProfiler/Widgets/STimersView.h
@@ -20,10 +20,15 @@
 #include "Insights/TimingProfiler/ViewModels/TimerGroupingAndSorting.h"
 #include "Insights/TimingProfiler/ViewModels/TimerNode.h"
 
+////////////////////////////////////////////////////////////////////////////////////////////////////
+
 class FMenuBuilder;
 class FUICommandList;
 
-namespace TraceServices { class IAnalysisSession; }
+namespace TraceServices
+{
+	class IAnalysisSession;
+}
 
 namespace UE::Insights
 {
@@ -40,6 +45,7 @@ enum class ESelectEventType : uint32;
 class SFrameTrack;
 class FTimerAggregator;
 class FTimingGraphTrack;
+class STimingProfilerWindow;
 class STimingView;
 
 ////////////////////////////////////////////////////////////////////////////////////////////////////
@@ -126,7 +132,7 @@ private:
 	void TreeView_BuildSortByMenu(FMenuBuilder& MenuBuilder);
 	void TreeView_BuildViewColumnMenu(FMenuBuilder& MenuBuilder);
 
-	//////////////////////////////////////////////////
+	////////////////////////////////////////////////////////////////////////////////////////////////////
 	// Copy to Clipboard
 
 	bool ContextMenu_CopyToClipboard_CanExecute() const;
@@ -135,7 +141,7 @@ private:
 	bool ContextMenu_CopyTimerNameToClipboard_CanExecute() const;
 	void ContextMenu_CopyTimerNameToClipboard_Execute();
 
-	//////////////////////////////////////////////////
+	////////////////////////////////////////////////////////////////////////////////////////////////////
 	// Export
 
 	void TreeView_BuildExportMenu(FMenuBuilder& MenuBuilder);
@@ -160,7 +166,7 @@ private:
 	bool OpenSaveTextFileDialog(const FString& InDialogTitle, const FString& InDefaultFile, FString& OutFilename) const;
 	class IFileHandle* OpenExportFile(const TCHAR* InFilename) const;
 
-	//////////////////////////////////////////////////
+	////////////////////////////////////////////////////////////////////////////////////////////////////
 	// Open Source File in IDE
 
 	bool ContextMenu_OpenSource_CanExecute() const;
@@ -168,7 +174,7 @@ private:
 
 	void OpenSourceFileInIDE(FTimerNodePtr InNode) const;
 
-	//////////////////////////////////////////////////
+	////////////////////////////////////////////////////////////////////////////////////////////////////
 	// Generate and Edit Timer Colors
 
 	bool ContextMenu_GenerateNewColor_CanExecute() const;
@@ -184,7 +190,7 @@ private:
 	void SetEditableColor(FLinearColor NewColor);
 	void ColorPickerCancelled(FLinearColor OriginalColor);
 
-	//////////////////////////////////////////////////
+	////////////////////////////////////////////////////////////////////////////////////////////////////
 	// Plot (Graph Series)
 
 	void TreeView_BuildPlotTimerMenu(FMenuBuilder& MenuBuilder);
@@ -231,9 +237,6 @@ private:
 
 	FTimerNodePtr GetSingleSelectedTimerNode() const;
 
-	/** Checks validity of one or more selected nodes*/
-	bool AreSelectedTimerNodesValid() const;
-
 	/** Called by STreeView when a tree item is double clicked. */
 	void TreeView_OnMouseButtonDoubleClick(FTimerNodePtr TreeNode);
 
@@ -289,20 +292,14 @@ private:
 	FText GroupBy_GetSelectedTooltipText() const;
 
 	////////////////////////////////////////////////////////////////////////////////////////////////////
-	// Frame Type
+	// Aggregation Mode
 
 	void CreateModeOptionsSources();
-
-	void Mode_OnSelectionChanged(TSharedPtr<ETraceFrameType> NewGroupingMode, ESelectInfo::Type SelectInfo);
-
-	TSharedRef<SWidget> Mode_OnGenerateWidget(TSharedPtr<ETraceFrameType> InGroupingMode) const;
-
+	void Mode_OnSelectionChanged(TSharedPtr<ETraceFrameType> InModeFrameType, ESelectInfo::Type SelectInfo);
+	TSharedRef<SWidget> Mode_OnGenerateWidget(TSharedPtr<ETraceFrameType> InModeFrameType) const;
 	FText Mode_GetSelectedText() const;
-
 	FText Mode_GetText(ETraceFrameType InFrameType) const;
-
 	FText Mode_GetSelectedTooltipText() const;
-
 	FText Mode_GetTooltipText(ETraceFrameType InFrameType) const;
 
 	////////////////////////////////////////////////////////////////////////////////////////////////////
@@ -408,10 +405,11 @@ private:
 	void SetTimingViewFrameType();
 
 	FText SetSourceFileSearchPaths_GetTooltip() const;
-	void  SetSourceFileSearchPaths_Execute();
-	
+	void SetSourceFileSearchPaths_Execute();
+
 	void ToggleShowFormatInTimerNames_Execute();
 
+	TSharedPtr<STimingProfilerWindow> GetProfilerWindow() const;
 	TSharedPtr<STimingView> GetTimingView() const;
 
 private:
