[Insights] Callers&Callees: Improved responsiveness when having large number of items selected.

#Synced-CL 50493319

diff --git a/Engine/Source/Developer/TraceInsights/Private/Insights/TimingProfiler/Widgets/SStatsView.cpp b/Engine/Source/Developer/TraceInsights/Private/Insights/TimingProfiler/Widgets/SStatsView.cpp
--- a/Engine/Source/Developer/TraceInsights/Private/Insights/TimingProfiler/Widgets/SStatsView.cpp
+++ b/Engine/Source/Developer/TraceInsights/Private/Insights/TimingProfiler/Widgets/SStatsView.cpp
@@ -406,8 +406,24 @@ END_SLATE_FUNCTION_BUILD_OPTIMIZATION
 
 TSharedPtr<SWidget> SStatsView::TreeView_GetMenuContent()
 {
-	const int32 NumSelectedNodes = TreeView->GetNumItemsSelected();
+	int32 NumSelectedNodes = TreeView->GetNumItemsSelected();
 	FStatsNodePtr SingleSelectedNode;
+	bool bIsValidNode = false;
+
+	if (NumSelectedNodes == 1)
+	{
+		TArray<FStatsNodePtr> SelectedNodes = TreeView->GetSelectedItems();
+		check(SelectedNodes.Num() == 1);
+		SingleSelectedNode = SelectedNodes[0];
+		if (SingleSelectedNode.IsValid())
+		{
+			bIsValidNode = (SingleSelectedNode->GetType() != EStatsNodeType::Group);
+		}
+		else
+		{
+			NumSelectedNodes = 0;
+		}
+	}
 
 	FText SelectionStr;
 	if (NumSelectedNodes == 0)
@@ -416,10 +432,7 @@ TSharedPtr<SWidget> SStatsView::TreeView_GetMenuContent()
 	}
 	else if (NumSelectedNodes == 1)
 	{
-		TArray<FStatsNodePtr> SelectedNodes = TreeView->GetSelectedItems();
-		check(SelectedNodes.Num() == 1);
-		SingleSelectedNode = SelectedNodes[0];
-
+		check(SingleSelectedNode.IsValid());
 		FString ItemName = SingleSelectedNode->GetName().ToString();
 		const int32 MaxStringLen = 64;
 		if (ItemName.Len() > MaxStringLen)
@@ -464,10 +477,6 @@ TSharedPtr<SWidget> SStatsView::TreeView_GetMenuContent()
 	// Counter options section
 	MenuBuilder.BeginSection("CounterOptions", LOCTEXT("ContextMenu_Section_CounterOptions", "Counter Options"));
 	{
-		const bool bIsValidNode =
-			SingleSelectedNode.IsValid() &&
-			SingleSelectedNode->GetType() != EStatsNodeType::Group;
-
 		auto CanExecute = [this, bIsValidNode]()
 		{
 			return bIsValidNode && GetTimingView().IsValid();
diff --git a/Engine/Source/Developer/TraceInsights/Private/Insights/TimingProfiler/Widgets/STimerTreeView.cpp b/Engine/Source/Developer/TraceInsights/Private/Insights/TimingProfiler/Widgets/STimerTreeView.cpp
--- a/Engine/Source/Developer/TraceInsights/Private/Insights/TimingProfiler/Widgets/STimerTreeView.cpp
+++ b/Engine/Source/Developer/TraceInsights/Private/Insights/TimingProfiler/Widgets/STimerTreeView.cpp
@@ -144,10 +144,10 @@ void STimerTreeView::InitCommandList()
 	CommandList->MapAction(FTimerTreeViewCommands::Get().Command_CopyToClipboard, FExecuteAction::CreateSP(this, &STimerTreeView::ContextMenu_CopyToClipboard_Execute), FCanExecuteAction::CreateSP(this, &STimerTreeView::ContextMenu_CopyToClipboard_CanExecute));
 	CommandList->MapAction(FTimerTreeViewCommands::Get().Command_CopyNameToClipboard, FExecuteAction::CreateSP(this, &STimerTreeView::ContextMenu_CopyTimerNameToClipboard_Execute), FCanExecuteAction::CreateSP(this, &STimerTreeView::ContextMenu_CopyTimerNameToClipboard_CanExecute));
 	CommandList->MapAction(FTimerTreeViewCommands::Get().Command_OpenSource, FExecuteAction::CreateSP(this, &STimerTreeView::ContextMenu_OpenSource_Execute), FCanExecuteAction::CreateSP(this, &STimerTreeView::ContextMenu_OpenSource_CanExecute));
-	CommandList->MapAction(FTimerTreeViewCommands::Get().Command_FindMaxInstance, FExecuteAction::CreateSP(this, &STimerTreeView::ContextMenu_FindInstance_Execute, true), FCanExecuteAction::CreateSP(this, &STimerTreeView::ContextMenu_FindInstance_CanExecute));
-	CommandList->MapAction(FTimerTreeViewCommands::Get().Command_FindMinInstance, FExecuteAction::CreateSP(this, &STimerTreeView::ContextMenu_FindInstance_Execute, false), FCanExecuteAction::CreateSP(this, &STimerTreeView::ContextMenu_FindInstance_CanExecute));
-	CommandList->MapAction(FTimerTreeViewCommands::Get().Command_FindMaxInstanceInSelection, FExecuteAction::CreateSP(this, &STimerTreeView::ContextMenu_FindInstanceInSelection_Execute, true), FCanExecuteAction::CreateSP(this, &STimerTreeView::ContextMenu_FindInstanceInSelection_CanExecute));
-	CommandList->MapAction(FTimerTreeViewCommands::Get().Command_FindMinInstanceInSelection, FExecuteAction::CreateSP(this, &STimerTreeView::ContextMenu_FindInstanceInSelection_Execute, false), FCanExecuteAction::CreateSP(this, &STimerTreeView::ContextMenu_FindInstanceInSelection_CanExecute));
+	CommandList->MapAction(FTimerTreeViewCommands::Get().Command_FindMaxInstance, FExecuteAction::CreateSP(this, &STimerTreeView::ContextMenu_FindInstance_Execute, ESelectEventType::Max), FCanExecuteAction::CreateSP(this, &STimerTreeView::ContextMenu_FindInstance_CanExecute));
+	CommandList->MapAction(FTimerTreeViewCommands::Get().Command_FindMinInstance, FExecuteAction::CreateSP(this, &STimerTreeView::ContextMenu_FindInstance_Execute, ESelectEventType::Min), FCanExecuteAction::CreateSP(this, &STimerTreeView::ContextMenu_FindInstance_CanExecute));
+	CommandList->MapAction(FTimerTreeViewCommands::Get().Command_FindMaxInstanceInSelection, FExecuteAction::CreateSP(this, &STimerTreeView::ContextMenu_FindInstanceInSelection_Execute, ESelectEventType::Max), FCanExecuteAction::CreateSP(this, &STimerTreeView::ContextMenu_FindInstanceInSelection_CanExecute));
+	CommandList->MapAction(FTimerTreeViewCommands::Get().Command_FindMinInstanceInSelection, FExecuteAction::CreateSP(this, &STimerTreeView::ContextMenu_FindInstanceInSelection_Execute, ESelectEventType::Min), FCanExecuteAction::CreateSP(this, &STimerTreeView::ContextMenu_FindInstanceInSelection_CanExecute));
 }
 
 ////////////////////////////////////////////////////////////////////////////////////////////////////
@@ -266,9 +266,24 @@ END_SLATE_FUNCTION_BUILD_OPTIMIZATION
 
 TSharedPtr<SWidget> STimerTreeView::TreeView_GetMenuContent()
 {
-	const TArray<FTimerNodePtr> SelectedNodes = TreeView->GetSelectedItems();
-	const int32 NumSelectedNodes = SelectedNodes.Num();
-	FTimerNodePtr SelectedNode = NumSelectedNodes ? SelectedNodes[0] : nullptr;
+	int32 NumSelectedNodes = TreeView->GetNumItemsSelected();
+	FTimerNodePtr SingleSelectedNode = nullptr;
+	bool bIsValidNode = false;
+
+	if (NumSelectedNodes == 1)
+	{
+		TArray<FTimerNodePtr> SelectedNodes = TreeView->GetSelectedItems();
+		check(SelectedNodes.Num() == 1);
+		SingleSelectedNode = SelectedNodes[0];
+		if (SingleSelectedNode.IsValid())
+		{
+			bIsValidNode = (SingleSelectedNode->GetType() != ETimerNodeType::Group);
+		}
+		else
+		{
+			NumSelectedNodes = 0;
+		}
+	}
 
 	FText SelectionStr;
 	if (NumSelectedNodes == 0)
@@ -277,7 +292,8 @@ TSharedPtr<SWidget> STimerTreeView::TreeView_GetMenuContent()
 	}
 	else if (NumSelectedNodes == 1)
 	{
-		FString ItemName = SelectedNode->GetName().ToString();
+		check(SingleSelectedNode.IsValid());
+		FString ItemName = SingleSelectedNode->GetName().ToString();
 		const int32 MaxStringLen = 64;
 		if (ItemName.Len() > MaxStringLen)
 		{
@@ -321,23 +337,22 @@ TSharedPtr<SWidget> STimerTreeView::TreeView_GetMenuContent()
 	// Timer options section
 	MenuBuilder.BeginSection("TimerOptions", LOCTEXT("ContextMenu_Section_TimerOptions", "Timer Options"));
 	{
-		TSharedPtr<STimingView> TimingView = GetTimingView();
-
-		auto CanExecute = [TimingView, NumSelectedNodes, SelectedNode]()
+		auto CanExecute = [this, bIsValidNode]()
 		{
-			return TimingView.IsValid() && NumSelectedNodes == 1 && SelectedNode.IsValid() && SelectedNode->GetType() != ETimerNodeType::Group;
+			return bIsValidNode && GetTimingView().IsValid();
 		};
 
 		// Highlight event
 		{
 			FUIAction Action_ToggleHighlight;
 			Action_ToggleHighlight.CanExecuteAction = FCanExecuteAction::CreateLambda(CanExecute);
-			Action_ToggleHighlight.ExecuteAction = FExecuteAction::CreateSP(this, &STimerTreeView::ToggleTimingViewEventFilter, SelectedNode);
+			Action_ToggleHighlight.ExecuteAction = FExecuteAction::CreateSP(this, &STimerTreeView::ToggleTimingViewEventFilter, SingleSelectedNode);
 
-			if (SelectedNode.IsValid() &&
-				SelectedNode->GetType() != ETimerNodeType::Group &&
+			TSharedPtr<STimingView> TimingView = GetTimingView();
+
+			if (bIsValidNode &&
 				TimingView.IsValid() &&
-				TimingView->IsFilterByEventType(SelectedNode->GetTimerId()))
+				TimingView->IsFilterByEventType(SingleSelectedNode->GetTimerId()))
 			{
 				MenuBuilder.AddMenuEntry
 				(
@@ -388,11 +403,7 @@ TSharedPtr<SWidget> STimerTreeView::TreeView_GetMenuContent()
 
 			FString File;
 			uint32 Line = 0;
-			bool bIsValidSource = false;
-			if (NumSelectedNodes == 1 && SelectedNode.IsValid() && SelectedNode->GetType() != ETimerNodeType::Group)
-			{
-				bIsValidSource = SelectedNode->GetSourceFileAndLine(File, Line);
-			}
+			bool bIsValidSource = bIsValidNode && SingleSelectedNode->GetSourceFileAndLine(File, Line);
 			FString UsablePath;
 			const bool bPathExists = FInsightsManager::Get()->GetSourceFilePathHelper().GetUsableFilePath(File, UsablePath);
 
@@ -570,22 +581,17 @@ void STimerTreeView::TreeView_BuildPlotTimerMenu(FMenuBuilder& MenuBuilder)
 {
 	MenuBuilder.SetSearchable(false);
 
-	const TArray<FTimerNodePtr> SelectedNodes = TreeView->GetSelectedItems();
-	const int32 NumSelectedNodes = SelectedNodes.Num();
-	FTimerNodePtr SelectedNode = NumSelectedNodes ? SelectedNodes[0] : nullptr;
+	FTimerNodePtr SelectedNode = GetSingleSelectedTimerNode();
+	const bool bIsValidNode = SelectedNode.IsValid() && SelectedNode->GetType() != ETimerNodeType::Group;
 
-	auto CanExecuteAddToGraphTrack = [NumSelectedNodes, SelectedNode]()
+	auto CanExecuteAddToGraphTrack = [this, bIsValidNode]()
 	{
-		TSharedPtr<STimingProfilerWindow> Wnd = FTimingProfilerManager::Get()->GetProfilerWindow();
-		TSharedPtr<STimingView> TimingView = Wnd.IsValid() ? Wnd->GetTimingView() : nullptr;
-		return TimingView.IsValid() && NumSelectedNodes == 1 && SelectedNode.IsValid() && SelectedNode->GetType() != ETimerNodeType::Group;
+		return bIsValidNode && GetTimingView().IsValid();
 	};
 
-	auto CanExecuteAddToFramesTrack = [NumSelectedNodes, SelectedNode]()
+	auto CanExecuteAddToFramesTrack = [this, bIsValidNode]()
 	{
-		TSharedPtr<STimingProfilerWindow> Wnd = FTimingProfilerManager::Get()->GetProfilerWindow();
-		TSharedPtr<SFrameTrack> FrameTrack = Wnd.IsValid() ? Wnd->GetFrameView() : nullptr;
-		return FrameTrack.IsValid() && NumSelectedNodes == 1 && SelectedNode.IsValid() && SelectedNode->GetType() != ETimerNodeType::Group;
+		return bIsValidNode && GetFrameTrack().IsValid();
 	};
 
 	MenuBuilder.BeginSection("Instance", LOCTEXT("Plot_Series_Instance_Section", "Instance"));
@@ -596,8 +602,7 @@ void STimerTreeView::TreeView_BuildPlotTimerMenu(FMenuBuilder& MenuBuilder)
 		Action_ToggleTimerInGraphTrack.CanExecuteAction = FCanExecuteAction::CreateLambda(CanExecuteAddToGraphTrack);
 		Action_ToggleTimerInGraphTrack.ExecuteAction = FExecuteAction::CreateSP(this, &STimerTreeView::ToggleTimingViewMainGraphEventInstanceSeries, SelectedNode);
 
-		if (SelectedNode.IsValid() &&
-			SelectedNode->GetType() != ETimerNodeType::Group &&
+		if (bIsValidNode &&
 			IsInstanceSeriesInTimingViewMainGraph(SelectedNode))
 		{
 			MenuBuilder.AddMenuEntry
@@ -634,8 +639,7 @@ void STimerTreeView::TreeView_BuildPlotTimerMenu(FMenuBuilder& MenuBuilder)
 		Action_ToggleFrameStatsTimerInGraphTrack.CanExecuteAction = FCanExecuteAction::CreateLambda(CanExecuteAddToGraphTrack);
 		Action_ToggleFrameStatsTimerInGraphTrack.ExecuteAction = FExecuteAction::CreateSP(this, &STimerTreeView::ToggleTimingViewMainGraphEventFrameStatsSeries, SelectedNode, ETraceFrameType::TraceFrameType_Game);
 
-		if (SelectedNode.IsValid() &&
-			SelectedNode->GetType() != ETimerNodeType::Group &&
+		if (bIsValidNode &&
 			IsFrameStatsSeriesInTimingViewMainGraph(SelectedNode, ETraceFrameType::TraceFrameType_Game))
 		{
 			MenuBuilder.AddMenuEntry
@@ -668,8 +672,7 @@ void STimerTreeView::TreeView_BuildPlotTimerMenu(FMenuBuilder& MenuBuilder)
 		Action_ToggleFrameStatsTimerInFrameTrack.CanExecuteAction = FCanExecuteAction::CreateLambda(CanExecuteAddToFramesTrack);
 		Action_ToggleFrameStatsTimerInFrameTrack.ExecuteAction = FExecuteAction::CreateSP(this, &STimerTreeView::ToggleFrameTrackSeries, SelectedNode, ETraceFrameType::TraceFrameType_Game);
 
-		if (SelectedNode.IsValid() &&
-			SelectedNode->GetType() != ETimerNodeType::Group &&
+		if (bIsValidNode &&
 			IsSeriesInFrameTrack(SelectedNode, ETraceFrameType::TraceFrameType_Game))
 		{
 			MenuBuilder.AddMenuEntry
@@ -706,8 +709,7 @@ void STimerTreeView::TreeView_BuildPlotTimerMenu(FMenuBuilder& MenuBuilder)
 		Action_ToggleFrameStatsTimerInGraphTrack.CanExecuteAction = FCanExecuteAction::CreateLambda(CanExecuteAddToGraphTrack);
 		Action_ToggleFrameStatsTimerInGraphTrack.ExecuteAction = FExecuteAction::CreateSP(this, &STimerTreeView::ToggleTimingViewMainGraphEventFrameStatsSeries, SelectedNode, ETraceFrameType::TraceFrameType_Rendering);
 
-		if (SelectedNode.IsValid() &&
-			SelectedNode->GetType() != ETimerNodeType::Group &&
+		if (bIsValidNode &&
 			IsFrameStatsSeriesInTimingViewMainGraph(SelectedNode, ETraceFrameType::TraceFrameType_Rendering))
 		{
 			MenuBuilder.AddMenuEntry
@@ -740,8 +742,7 @@ void STimerTreeView::TreeView_BuildPlotTimerMenu(FMenuBuilder& MenuBuilder)
 		Action_ToggleFrameStatsTimerInFrameTrack.CanExecuteAction = FCanExecuteAction::CreateLambda(CanExecuteAddToFramesTrack);
 		Action_ToggleFrameStatsTimerInFrameTrack.ExecuteAction = FExecuteAction::CreateSP(this, &STimerTreeView::ToggleFrameTrackSeries, SelectedNode, ETraceFrameType::TraceFrameType_Rendering);
 
-		if (SelectedNode.IsValid() &&
-			SelectedNode->GetType() != ETimerNodeType::Group &&
+		if (bIsValidNode &&
 			IsSeriesInFrameTrack(SelectedNode, ETraceFrameType::TraceFrameType_Rendering))
 		{
 			MenuBuilder.AddMenuEntry
@@ -1005,11 +1006,8 @@ FTimerNodePtr STimerTreeView::GetSingleSelectedTimerNode() const
 		return nullptr;
 	}
 	const TArray<FTimerNodePtr> SelectedNodes = TreeView->GetSelectedItems();
-	if (SelectedNodes.Num() == 1)
-	{
-		return SelectedNodes[0];
-	}
-	return nullptr;
+	check(SelectedNodes.Num() == 1);
+	return SelectedNodes[0];
 }
 
 ////////////////////////////////////////////////////////////////////////////////////////////////////
@@ -1702,7 +1700,7 @@ TSharedPtr<FTimingGraphTrack> STimerTreeView::GetTimingViewMainGraphTrack() cons
 
 TSharedPtr<SFrameTrack> STimerTreeView::GetFrameTrack() const
 {
-	TSharedPtr<STimingProfilerWindow> Wnd = FTimingProfilerManager::Get()->GetProfilerWindow();
+	TSharedPtr<STimingProfilerWindow> Wnd = GetProfilerWindow();
 	return Wnd.IsValid() ? Wnd->GetFrameView() : nullptr;
 }
 
@@ -1726,7 +1724,7 @@ void STimerTreeView::ToggleGraphInstanceSeries(TSharedRef<FTimingGraphTrack> Gra
 		GraphTrack->SetDirtyFlag();
 	}
 
-	TSharedPtr<STimingProfilerWindow> ProfilerWindow = FTimingProfilerManager::Get()->GetProfilerWindow();
+	TSharedPtr<STimingProfilerWindow> ProfilerWindow = GetProfilerWindow();
 	if (ProfilerWindow.IsValid())
 	{
 		ProfilerWindow->OnTimerAddedToGraphsChanged(TimerId);
@@ -1784,7 +1782,7 @@ void STimerTreeView::ToggleGraphFrameStatsSeries(TSharedRef<FTimingGraphTrack> G
 		GraphTrack->SetDirtyFlag();
 	}
 
-	TSharedPtr<STimingProfilerWindow> ProfilerWindow = FTimingProfilerManager::Get()->GetProfilerWindow();
+	TSharedPtr<STimingProfilerWindow> ProfilerWindow = GetProfilerWindow();
 	if (ProfilerWindow.IsValid())
 	{
 		ProfilerWindow->OnTimerAddedToGraphsChanged(TimerId);
@@ -1876,7 +1874,7 @@ void STimerTreeView::ToggleFrameTrackSeries(FTimerNodePtr TimerNode, ETraceFrame
 		FrameTrack->AddTimerFrameStatSeries(FrameType, TimerId, TimerNode->GetColor(), SeriesName);
 	}
 
-	TSharedPtr<STimingProfilerWindow> ProfilerWindow = FTimingProfilerManager::Get()->GetProfilerWindow();
+	TSharedPtr<STimingProfilerWindow> ProfilerWindow = GetProfilerWindow();
 	if (ProfilerWindow.IsValid())
 	{
 		ProfilerWindow->OnTimerAddedToGraphsChanged(TimerId);
@@ -1889,9 +1887,7 @@ void STimerTreeView::ToggleFrameTrackSeries(FTimerNodePtr TimerNode, ETraceFrame
 
 bool STimerTreeView::ContextMenu_CopyToClipboard_CanExecute() const
 {
-	const TArray<FTimerNodePtr> SelectedNodes = TreeView->GetSelectedItems();
-
-	return SelectedNodes.Num() > 0;
+	return TreeView->GetNumItemsSelected() > 0;
 }
 
 ////////////////////////////////////////////////////////////////////////////////////////////////////
@@ -1909,7 +1905,7 @@ void STimerTreeView::ContextMenu_CopyToClipboard_Execute()
 
 bool STimerTreeView::ContextMenu_CopyTimerNameToClipboard_CanExecute() const
 {
-	return TreeView->GetSelectedItems().Num() > 0;
+	return TreeView->GetNumItemsSelected() > 0;
 }
 
 ////////////////////////////////////////////////////////////////////////////////////////////////////
@@ -2002,7 +1998,7 @@ bool STimerTreeView::ContextMenu_FindInstance_CanExecute() const
 
 ////////////////////////////////////////////////////////////////////////////////////////////////////
 
-void STimerTreeView::ContextMenu_FindInstance_Execute(bool bFindMax) const
+void STimerTreeView::ContextMenu_FindInstance_Execute(ESelectEventType EventType) const
 {
 	FTimerNodePtr SelectedNode = GetSingleSelectedTimerNode();
 	if (!SelectedNode.IsValid() || SelectedNode->GetType() == ETimerNodeType::Group)
@@ -2016,8 +2012,7 @@ void STimerTreeView::ContextMenu_FindInstance_Execute(bool bFindMax) const
 		return;
 	}
 
-	ESelectEventType Type = bFindMax ? ESelectEventType::Max : ESelectEventType::Min;
-	TimingView->SelectEventInstance(SelectedNode->GetTimerId(), Type, false);
+	TimingView->SelectEventInstance(SelectedNode->GetTimerId(), EventType, false);
 }
 
 ////////////////////////////////////////////////////////////////////////////////////////////////////
@@ -2041,7 +2036,7 @@ bool STimerTreeView::ContextMenu_FindInstanceInSelection_CanExecute() const
 
 ////////////////////////////////////////////////////////////////////////////////////////////////////
 
-void STimerTreeView::ContextMenu_FindInstanceInSelection_Execute(bool bFindMax) const
+void STimerTreeView::ContextMenu_FindInstanceInSelection_Execute(ESelectEventType EventType) const
 {
 	FTimerNodePtr SelectedNode = GetSingleSelectedTimerNode();
 	if (!SelectedNode.IsValid() || SelectedNode->GetType() == ETimerNodeType::Group)
@@ -2055,16 +2050,23 @@ void STimerTreeView::ContextMenu_FindInstanceInSelection_Execute(bool bFindMax)
 		return;
 	}
 
-	ESelectEventType Type = bFindMax ? ESelectEventType::Max : ESelectEventType::Min;
-	TimingView->SelectEventInstance(SelectedNode->GetTimerId(), Type, true);
+	TimingView->SelectEventInstance(SelectedNode->GetTimerId(), EventType, true);
 }
 
 ////////////////////////////////////////////////////////////////////////////////////////////////////
+// Misc
+////////////////////////////////////////////////////////////////////////////////////////////////////
+
+TSharedPtr<STimingProfilerWindow> STimerTreeView::GetProfilerWindow() const
+{
+	return FTimingProfilerManager::Get()->GetProfilerWindow();
+}
+
 ////////////////////////////////////////////////////////////////////////////////////////////////////
 
 TSharedPtr<STimingView> STimerTreeView::GetTimingView() const
 {
-	TSharedPtr<STimingProfilerWindow> ProfilerWindow = FTimingProfilerManager::Get()->GetProfilerWindow();
+	TSharedPtr<STimingProfilerWindow> ProfilerWindow = GetProfilerWindow();
 	return ProfilerWindow.IsValid() ? ProfilerWindow->GetTimingView() : nullptr;
 }
 
diff --git a/Engine/Source/Developer/TraceInsights/Private/Insights/TimingProfiler/Widgets/STimerTreeView.h b/Engine/Source/Developer/TraceInsights/Private/Insights/TimingProfiler/Widgets/STimerTreeView.h
--- a/Engine/Source/Developer/TraceInsights/Private/Insights/TimingProfiler/Widgets/STimerTreeView.h
+++ b/Engine/Source/Developer/TraceInsights/Private/Insights/TimingProfiler/Widgets/STimerTreeView.h
@@ -34,8 +34,10 @@ namespace UE::Insights
 namespace UE::Insights::TimingProfiler
 {
 
+enum class ESelectEventType : uint32;
 class SFrameTrack;
 class FTimingGraphTrack;
+class STimingProfilerWindow;
 class STimingView;
 
 ////////////////////////////////////////////////////////////////////////////////////////////////////
@@ -81,6 +83,7 @@ public:
 
 private:
 	void InitCommandList();
+
 	FTimerNodePtr CreateTimerNodeRec(const TraceServices::FTimingProfilerButterflyNode& Node, const TraceServices::ITimingProfilerProvider& TimingProfilerProvider);
 	void ExpandNodesRec(FTimerNodePtr NodePtr, int32 Depth);
 
@@ -125,10 +128,10 @@ private:
 	void TreeView_FindMenu(FMenuBuilder& MenuBuilder);
 
 	bool ContextMenu_FindInstance_CanExecute() const;
-	void ContextMenu_FindInstance_Execute(bool bFindMax) const;
+	void ContextMenu_FindInstance_Execute(ESelectEventType EventType) const;
 
 	bool ContextMenu_FindInstanceInSelection_CanExecute() const;
-	void ContextMenu_FindInstanceInSelection_Execute(bool bFindMax) const;
+	void ContextMenu_FindInstanceInSelection_Execute(ESelectEventType EventType) const;
 
 	////////////////////////////////////////////////////////////////////////////////////////////////////
 	// Tree View - Columns' Header
@@ -255,6 +258,7 @@ private:
 
 	////////////////////////////////////////////////////////////////////////////////////////////////////
 
+	TSharedPtr<STimingProfilerWindow> GetProfilerWindow() const;
 	TSharedPtr<STimingView> GetTimingView() const;
 
 	FTimerNodePtr GetTimerNodeRec(uint32 TimerId, const FTimerNodePtr TimerNode) const;
diff --git a/Engine/Source/Developer/TraceInsights/Private/Insights/TimingProfiler/Widgets/STimersView.cpp b/Engine/Source/Developer/TraceInsights/Private/Insights/TimingProfiler/Widgets/STimersView.cpp
--- a/Engine/Source/Developer/TraceInsights/Private/Insights/TimingProfiler/Widgets/STimersView.cpp
+++ b/Engine/Source/Developer/TraceInsights/Private/Insights/TimingProfiler/Widgets/STimersView.cpp
@@ -603,8 +603,24 @@ END_SLATE_FUNCTION_BUILD_OPTIMIZATION
 
 TSharedPtr<SWidget> STimersView::TreeView_GetMenuContent()
 {
-	const int32 NumSelectedNodes = TreeView->GetNumItemsSelected();
+	int32 NumSelectedNodes = TreeView->GetNumItemsSelected();
 	FTimerNodePtr SingleSelectedNode;
+	bool bIsValidNode = false;
+
+	if (NumSelectedNodes == 1)
+	{
+		TArray<FTimerNodePtr> SelectedNodes = TreeView->GetSelectedItems();
+		check(SelectedNodes.Num() == 1);
+		SingleSelectedNode = SelectedNodes[0];
+		if (SingleSelectedNode.IsValid())
+		{
+			bIsValidNode = (SingleSelectedNode->GetType() != ETimerNodeType::Group);
+		}
+		else
+		{
+			NumSelectedNodes = 0;
+		}
+	}
 
 	FText SelectionStr;
 	if (NumSelectedNodes == 0)
@@ -613,10 +629,7 @@ TSharedPtr<SWidget> STimersView::TreeView_GetMenuContent()
 	}
 	else if (NumSelectedNodes == 1)
 	{
-		TArray<FTimerNodePtr> SelectedNodes = TreeView->GetSelectedItems();
-		check(SelectedNodes.Num() == 1);
-		SingleSelectedNode = SelectedNodes[0];
-
+		check(SingleSelectedNode.IsValid());
 		FString ItemName = SingleSelectedNode->GetName().ToString();
 		const int32 MaxStringLen = 64;
 		if (ItemName.Len() > MaxStringLen)
@@ -661,10 +674,6 @@ TSharedPtr<SWidget> STimersView::TreeView_GetMenuContent()
 	// Timer options section
 	MenuBuilder.BeginSection("TimerOptions", LOCTEXT("ContextMenu_Section_TimerOptions", "Timer Options"));
 	{
-		const bool bIsValidNode =
-			SingleSelectedNode.IsValid() &&
-			SingleSelectedNode->GetType() != ETimerNodeType::Group;
-
 		auto CanExecute = [this, bIsValidNode]()
 		{
 			return bIsValidNode && GetTimingView().IsValid();
@@ -731,11 +740,7 @@ TSharedPtr<SWidget> STimersView::TreeView_GetMenuContent()
 
 			FString File;
 			uint32 Line = 0;
-			bool bIsValidSource = false;
-			if (SingleSelectedNode.IsValid() && SingleSelectedNode->GetType() != ETimerNodeType::Group)
-			{
-				bIsValidSource = SingleSelectedNode->GetSourceFileAndLine(File, Line);
-			}
+			bool bIsValidSource = bIsValidNode && SingleSelectedNode->GetSourceFileAndLine(File, Line);
 			FString UsablePath;
 			const bool bPathExists = FInsightsManager::Get()->GetSourceFilePathHelper().GetUsableFilePath(File, UsablePath);
 
