[Insights]
- Added a new FUtf8FileWriter utility (in Private/Insights/Common/) and replaced FUtf8Writer from TimingExporter and MemoryExporter to use this implementation.
- Converted all UTF16 exporters to use UTF8 (2x speed improvement).
- MemoryInsights: Added support for all investigation rules for the MemoryInsights.ExportAllocs execute command. Also made improvements to file writing speed. Added progress logs (including total export duration, in seconds).

#Synced-CL 50589729

diff --git a/Engine/Source/Developer/TraceInsights/Private/Insights/Common/Utf8FileWriter.cpp b/Engine/Source/Developer/TraceInsights/Private/Insights/Common/Utf8FileWriter.cpp
new file mode 100644
--- /dev/null
+++ b/Engine/Source/Developer/TraceInsights/Private/Insights/Common/Utf8FileWriter.cpp
@@ -0,0 +1,237 @@
+// Copyright Epic Games, Inc. All Rights Reserved.
+
+#include "Utf8FileWriter.h"
+
+#include "Containers/Utf8String.h"
+#include "GenericPlatform/GenericPlatformFile.h"
+#include "String/Find.h"
+
+#include <limits>
+#include <cmath>
+
+namespace UE::Insights
+{
+
+////////////////////////////////////////////////////////////////////////////////////////////////////
+
+FUtf8FileWriter::FUtf8FileWriter(IFileHandle* InFileHandle, bool bIsCSV)
+	: FileHandle(InFileHandle)
+	, Separator(bIsCSV ? UTF8CHAR(',') : UTF8CHAR('\t'))
+{
+	check(FileHandle);
+}
+
+////////////////////////////////////////////////////////////////////////////////////////////////////
+
+void FUtf8FileWriter::AppendValue(const UTF8CHAR* Str)
+{
+	if (Str != nullptr)
+	{
+		AppendValue(FUtf8StringView(Str));
+	}
+}
+
+////////////////////////////////////////////////////////////////////////////////////////////////////
+
+void FUtf8FileWriter::AppendValue(FUtf8StringView Str)
+{
+	if (Str.IsEmpty())
+	{
+		// nothing to append
+		return;
+	}
+
+	if (IsCSV())
+	{
+		bool bNeedsQuotes = false;
+		if (UE::String::FindFirstChar(Str, Separator, ESearchCase::CaseSensitive) != INDEX_NONE)
+		{
+			AppendDoubleQuotationMark();
+			bNeedsQuotes = true;
+		}
+		if (UE::String::FindFirstChar(Str, UTF8CHAR('\"'), ESearchCase::CaseSensitive) != INDEX_NONE)
+		{
+			if (!bNeedsQuotes)
+			{
+				AppendDoubleQuotationMark();
+				bNeedsQuotes = true;
+			}
+			FUtf8String String{ Str };
+			String.ReplaceInline(UTF8TEXT("\""), UTF8TEXT("\"\"")); // CSV escape for " is ""
+			StringBuilder.Append(String);
+		}
+		else
+		{
+			StringBuilder.Append(Str);
+		}
+		if (bNeedsQuotes)
+		{
+			AppendDoubleQuotationMark();
+		}
+	}
+	else // TSV
+	{
+		FUtf8String String{ Str };
+		String.ReplaceCharWithEscapedCharInline();
+		StringBuilder.Append(String);
+	}
+}
+
+////////////////////////////////////////////////////////////////////////////////////////////////////
+
+void FUtf8FileWriter::AppendValue(const TCHAR* Str)
+{
+	if (Str != nullptr)
+	{
+		AppendValue(FStringView(Str));
+	}
+}
+
+////////////////////////////////////////////////////////////////////////////////////////////////////
+
+void FUtf8FileWriter::AppendValue(FStringView Str)
+{
+	if (Str.IsEmpty())
+	{
+		// nothing to append
+		return;
+	}
+
+	if (IsCSV())
+	{
+		bool bNeedsQuotes = false;
+		if (UE::String::FindFirstChar(Str, Separator, ESearchCase::CaseSensitive) != INDEX_NONE)
+		{
+			AppendDoubleQuotationMark();
+			bNeedsQuotes = true;
+		}
+		if (UE::String::FindFirstChar(Str, TCHAR('\"'), ESearchCase::CaseSensitive) != INDEX_NONE)
+		{
+			if (!bNeedsQuotes)
+			{
+				AppendDoubleQuotationMark();
+				bNeedsQuotes = true;
+			}
+			FString String{ Str };
+			String.ReplaceInline(TEXT("\""), TEXT("\"\"")); // CSV escape for " is ""
+			StringBuilder.Append(String);
+		}
+		else
+		{
+			StringBuilder.Append(Str);
+		}
+		if (bNeedsQuotes)
+		{
+			AppendDoubleQuotationMark();
+		}
+	}
+	else // TSV
+	{
+		FString String{ Str };
+		String.ReplaceCharWithEscapedCharInline();
+		StringBuilder.Append(String);
+	}
+}
+
+////////////////////////////////////////////////////////////////////////////////////////////////////
+
+void FUtf8FileWriter::AppendValue(float InValue)
+{
+	if (std::isnan(InValue))
+	{
+		StringBuilder.Append(UTF8TEXTVIEW("nan"));
+	}
+	else if (InValue == +std::numeric_limits<float>::infinity() || InValue == +FLT_MAX)
+	{
+		StringBuilder.Append(UTF8TEXTVIEW("+inf"));
+	}
+	else if (InValue == -std::numeric_limits<float>::infinity() || InValue == -FLT_MAX)
+	{
+		StringBuilder.Append(UTF8TEXTVIEW("-inf"));
+	}
+	else
+	{
+		StringBuilder.Appendf(UTF8TEXT("%.6g"), InValue);
+	}
+}
+
+////////////////////////////////////////////////////////////////////////////////////////////////////
+
+void FUtf8FileWriter::AppendValue(double InValue)
+{
+	if (std::isnan(InValue))
+	{
+		StringBuilder.Append(UTF8TEXTVIEW("nan"));
+	}
+	else if (InValue == +std::numeric_limits<double>::infinity() || InValue == +DBL_MAX)
+	{
+		StringBuilder.Append(UTF8TEXTVIEW("+inf"));
+	}
+	else if (InValue == -std::numeric_limits<double>::infinity() || InValue == -DBL_MAX)
+	{
+		StringBuilder.Append(UTF8TEXTVIEW("-inf"));
+	}
+	else
+	{
+		StringBuilder.Appendf(UTF8TEXT("%.9g"), InValue);
+	}
+}
+
+////////////////////////////////////////////////////////////////////////////////////////////////////
+
+void FUtf8FileWriter::AppendValueAsTimestamp(double InValue)
+{
+	if (std::isnan(InValue))
+	{
+		StringBuilder.Append(UTF8TEXTVIEW("nan"));
+	}
+	else if (InValue == +std::numeric_limits<double>::infinity() || InValue == +DBL_MAX)
+	{
+		StringBuilder.Append(UTF8TEXTVIEW("+inf"));
+	}
+	else if (InValue == -std::numeric_limits<double>::infinity() || InValue == -DBL_MAX)
+	{
+		StringBuilder.Append(UTF8TEXTVIEW("-inf"));
+	}
+	else
+	{
+		StringBuilder.Appendf(UTF8TEXT("%.9g"), InValue);
+	}
+}
+
+////////////////////////////////////////////////////////////////////////////////////////////////////
+
+void FUtf8FileWriter::AppendValueAsDuration(double InValue)
+{
+	if (std::isnan(InValue))
+	{
+		StringBuilder.Append(UTF8TEXTVIEW("nan"));
+	}
+	else if (InValue == +std::numeric_limits<double>::infinity() || InValue == +DBL_MAX)
+	{
+		StringBuilder.Append(UTF8TEXTVIEW("+inf"));
+	}
+	else if (InValue == -std::numeric_limits<double>::infinity() || InValue == -DBL_MAX)
+	{
+		StringBuilder.Append(UTF8TEXTVIEW("-inf"));
+	}
+	else
+	{
+		StringBuilder.Appendf(UTF8TEXT("%.9f"), InValue);
+	}
+}
+
+////////////////////////////////////////////////////////////////////////////////////////////////////
+
+void FUtf8FileWriter::WriteStringBuilder(int32 CacheLen)
+{
+	if (StringBuilder.Len() > CacheLen)
+	{
+		FileHandle->Write((const uint8*)StringBuilder.ToString(), StringBuilder.Len() * sizeof(UTF8CHAR));
+		StringBuilder.Reset();
+	}
+}
+
+////////////////////////////////////////////////////////////////////////////////////////////////////
+
+} // namespace UE::Insights
diff --git a/Engine/Source/Developer/TraceInsights/Private/Insights/Common/Utf8FileWriter.h b/Engine/Source/Developer/TraceInsights/Private/Insights/Common/Utf8FileWriter.h
new file mode 100644
--- /dev/null
+++ b/Engine/Source/Developer/TraceInsights/Private/Insights/Common/Utf8FileWriter.h
@@ -0,0 +1,122 @@
+// Copyright Epic Games, Inc. All Rights Reserved.
+
+#pragma once
+
+#include "CoreTypes.h"
+
+#include "Containers/StringView.h"
+#include "Containers/UnrealString.h"
+#include "Misc/StringBuilder.h"
+
+////////////////////////////////////////////////////////////////////////////////////////////////////
+
+class IFileHandle;
+
+////////////////////////////////////////////////////////////////////////////////////////////////////
+
+namespace UE::Insights
+{
+
+////////////////////////////////////////////////////////////////////////////////////////////////////
+
+class FUtf8FileWriter
+{
+private:
+	static constexpr int BufferSize = 64 * 1024;
+
+public:
+	FUtf8FileWriter() = delete;
+	FUtf8FileWriter(IFileHandle* InFileHandle, bool bIsCSV);
+	~FUtf8FileWriter() = default;
+
+	bool IsCSV() const { return Separator == UTF8CHAR(','); }
+	bool IsTSV() const { return Separator == UTF8CHAR('\t'); }
+
+	FUtf8StringBuilderBase& GetStringBuilder() { return StringBuilder; }
+
+	UTF8CHAR GetSeparator() const { return Separator; }
+	UTF8CHAR GetQuotationMark() const { return UTF8CHAR('\"'); }
+	UTF8CHAR GetLineEnd() const { return UTF8CHAR('\n'); }
+
+	//////////////////////////////////////////////////
+	// Writes the content without escaping.
+
+	void AppendSeparator() { StringBuilder.AppendChar(Separator); }
+
+	void AppendSingleQuotationMark() { StringBuilder.AppendChar(UTF8CHAR('\'')); }
+	void AppendDoubleQuotationMark() { StringBuilder.AppendChar(UTF8CHAR('\"')); }
+
+	void AppendLineEnd()
+	{
+		StringBuilder.AppendChar(UTF8CHAR('\n'));
+		WriteStringBuilder(BufferSize - 1024);
+	}
+
+	template <typename OtherCharType>
+	inline void Append(const OtherCharType* const String, const int32 Length)
+	{
+		StringBuilder.Append(String, Length);
+	}
+
+	template <typename CharRangeType>
+	inline void Append(CharRangeType&& Range)
+	{
+		StringBuilder.Append(Range);
+	}
+
+	template <typename AppendedCharType>
+	inline void AppendChar(AppendedCharType Char)
+	{
+		StringBuilder.Append(Char);
+	}
+
+	template <typename FmtType, typename... Types>
+	inline void Appendf(const FmtType& Fmt, Types... Args)
+	{
+		StringBuilder.Appendf(Fmt, Forward<Types>(Args)...);
+	}
+
+	inline void AppendV(const UTF8CHAR* Fmt, va_list Args)
+	{
+		StringBuilder.AppendV(Fmt, Args);
+	}
+
+	//////////////////////////////////////////////////
+	// Escapes the value (if needed; see CSV/TSV) before writing it to file.
+
+	void AppendValue(bool InValue)   { StringBuilder.Append(InValue ? UTF8TEXT("True") : UTF8TEXT("False")); }
+	void AppendValue(int32 InValue)  { StringBuilder.Appendf(UTF8TEXT("%i"), InValue); }
+	void AppendValue(uint32 InValue) { StringBuilder.Appendf(UTF8TEXT("%u"), InValue); }
+	void AppendValue(int64 InValue)  { StringBuilder.Appendf(UTF8TEXT("%lli"), InValue); }
+	void AppendValue(uint64 InValue) { StringBuilder.Appendf(UTF8TEXT("%llu"), InValue); }
+
+	void AppendValueAsHex32(uint32 InValue) { StringBuilder.Appendf(UTF8TEXT("0x%X"), InValue); }
+	void AppendValueAsHex64(uint64 InValue) { StringBuilder.Appendf(UTF8TEXT("0x%llX"), InValue); }
+
+	void AppendValue(float InValue);
+	void AppendValue(double InValue);
+	void AppendValueAsTimestamp(double InValue);
+	void AppendValueAsDuration(double InValue);
+
+	void AppendValue(const UTF8CHAR* InValue);
+	void AppendValue(FUtf8StringView InValue);
+	void AppendValue(const TCHAR* InValue);
+	void AppendValue(FStringView InValue);
+
+	//////////////////////////////////////////////////
+
+	void Flush()
+	{
+		WriteStringBuilder(0);
+	}
+
+private:
+	void WriteStringBuilder(int32 CacheLen);
+
+private:
+	IFileHandle* FileHandle = nullptr;
+	TUtf8StringBuilder<BufferSize> StringBuilder;
+	UTF8CHAR Separator;
+};
+
+} // namespace UE::Insights
diff --git a/Engine/Source/Developer/TraceInsights/Private/Insights/InsightsManager.cpp b/Engine/Source/Developer/TraceInsights/Private/Insights/InsightsManager.cpp
--- a/Engine/Source/Developer/TraceInsights/Private/Insights/InsightsManager.cpp
+++ b/Engine/Source/Developer/TraceInsights/Private/Insights/InsightsManager.cpp
@@ -769,8 +769,6 @@ void FInsightsManager::LoadTrace(uint32 InTraceId)
 	const UE::Trace::FStoreClient::FTraceInfo* TraceInfo = StoreClient->GetTraceInfoById(InTraceId);
 	if (TraceInfo != nullptr)
 	{
-		const FUtf8StringView Utf8NameView = TraceInfo->GetName();
-		FString Name(Utf8NameView);
 		FUtf8StringView Uri = TraceInfo->GetUri();
 		if (Uri.Len() > 0)
 		{
@@ -780,6 +778,7 @@ void FInsightsManager::LoadTrace(uint32 InTraceId)
 		{
 			// Fallback for older versions of UTS which didn't write uri
 			FString StoreDirectory(StoreClient->GetStatus()->GetStoreDir());
+			FString Name(TraceInfo->GetName());
 			TraceName = FPaths::SetExtension(FPaths::Combine(StoreDirectory, Name), TEXT(".utrace"));
 			FPaths::MakePlatformFilename(TraceName);
 		}
diff --git a/Engine/Source/Developer/TraceInsights/Private/Insights/MemoryProfiler/MemoryExportCommandParser.cpp b/Engine/Source/Developer/TraceInsights/Private/Insights/MemoryProfiler/MemoryExportCommandParser.cpp
--- a/Engine/Source/Developer/TraceInsights/Private/Insights/MemoryProfiler/MemoryExportCommandParser.cpp
+++ b/Engine/Source/Developer/TraceInsights/Private/Insights/MemoryProfiler/MemoryExportCommandParser.cpp
@@ -2,16 +2,42 @@
 
 #include "MemoryExportCommandParser.h"
 
+#include "Misc/Parse.h"
+
 namespace UE::Insights::MemoryProfiler
 {
+
+////////////////////////////////////////////////////////////////////////////////////////////////////
+
 DEFINE_LOG_CATEGORY(LogMemoryCommandParser);
 
+////////////////////////////////////////////////////////////////////////////////////////////////////
+
 FMemoryExportCommandParser::FMemoryExportCommandParser() :
 	TimeA(-1.0), TimeB(-1.0), TimeC(-1.0), TimeD(-1.0), MaxResults(0)
 {
+}
 
+////////////////////////////////////////////////////////////////////////////////////////////////////
+
+static bool ParseTimeParameter(double& OutTime, const FString& Value, const TCHAR* TimeName)
+{
+	if (!Value.IsNumeric())
+	{
+		UE_LOG(LogMemoryCommandParser, Error, TEXT("Time%s value (%s) is not a valid number."), TimeName, *Value);
+		return false;
+	}
+	OutTime = FCString::Atod(*Value);
+	if (OutTime < 0.0)
+	{
+		UE_LOG(LogMemoryCommandParser, Error, TEXT("Cannot pass a negative value (%.2f) for Time%s."), OutTime, TimeName);
+		return false;
+	}
+	return true;
 }
 
+////////////////////////////////////////////////////////////////////////////////////////////////////
+
 bool FMemoryExportCommandParser::ParseParameters(const TCHAR* Cmd, FOutputDevice& Ar)
 {
 	UE_LOG(LogMemoryCommandParser, Log, TEXT("Validating parameters necessary to execute chosen rule..."));
@@ -42,19 +68,19 @@ bool FMemoryExportCommandParser::ParseParameters(const TCHAR* Cmd, FOutputDevice
 				Token.TrimQuotesInline();
 				if (Token.IsEmpty())
 				{
-					UE_LOG(LogMemoryCommandParser, Error, TEXT("Exporter stopped because no Rule was provided. Please, specify one. Ex: -Rule=AaBCf"));
+					UE_LOG(LogMemoryCommandParser, Error, TEXT("Exporter stopped because no Rule was provided! Please, specify one. Ex: -Rule=AaBCf"));
 					return false;
 				}
-				const QueryRule* RuleEnumPtr = QueryRuleMap().Find(Token);
+				const EQueryRule* RuleEnumPtr = QueryRuleMap().Find(Token);
 				if (RuleEnumPtr == nullptr)
 				{
-					UE_LOG(LogMemoryCommandParser, Error, TEXT("Exporter stopped because Rule (%s) is not a valid one"), *Token);
+					UE_LOG(LogMemoryCommandParser, Error, TEXT("Exporter stopped because Rule (%s) is not valid!"), *Token);
 					TArray<FString> RuleNames;
 					QueryRuleMap().GetKeys(RuleNames);
-					UE_LOG(LogMemoryCommandParser, Error, TEXT("Try one of the following Rules only: [%s]"), *FString::Join(RuleNames, TEXT(" | ")));
+					UE_LOG(LogMemoryCommandParser, Error, TEXT("Available rules: %s"), *FString::Join(RuleNames, TEXT(", ")));
 					return false;
 				}
-				Rule = Token;
+				RuleName = Token;
 				RuleEnum = *RuleEnumPtr;
 			}
 			else if (Token.StartsWith(OutputToken))
@@ -94,15 +120,8 @@ bool FMemoryExportCommandParser::ParseParameters(const TCHAR* Cmd, FOutputDevice
 				{
 					continue;
 				}
-				if (!Token.IsNumeric())
-				{
-					UE_LOG(LogMemoryCommandParser, Error, TEXT("TimeA value %s is not a valid number"), *Token);
-					return false;
-				}
-				TimeA = FCString::Atod(*Token);
-				if (TimeA < 0.0)
+				if (!ParseTimeParameter(TimeA, Token, TEXT("A")))
 				{
-					UE_LOG(LogMemoryCommandParser, Error, TEXT("Cannot pass a negative value (%.2f) for TimeA"), TimeA);
 					return false;
 				}
 			}
@@ -113,15 +132,8 @@ bool FMemoryExportCommandParser::ParseParameters(const TCHAR* Cmd, FOutputDevice
 				{
 					continue;
 				}
-				if (!Token.IsNumeric())
+				if (!ParseTimeParameter(TimeB, Token, TEXT("B")))
 				{
-					UE_LOG(LogMemoryCommandParser, Error, TEXT("TimeB value %s is not a valid number"), *Token);
-					return false;
-				}
-				TimeB = FCString::Atod(*Token);
-				if (TimeB < 0.0)
-				{
-					UE_LOG(LogMemoryCommandParser, Error, TEXT("Cannot pass a negative value (%.2f) for TimeB"), TimeB);
 					return false;
 				}
 			}
@@ -132,15 +144,8 @@ bool FMemoryExportCommandParser::ParseParameters(const TCHAR* Cmd, FOutputDevice
 				{
 					continue;
 				}
-				if (!Token.IsNumeric())
+				if (!ParseTimeParameter(TimeC, Token, TEXT("C")))
 				{
-					UE_LOG(LogMemoryCommandParser, Error, TEXT("TimeC value %s is not a valid number"), *Token);
-					return false;
-				}
-				TimeC = FCString::Atod(*Token);
-				if (TimeC < 0.0)
-				{
-					UE_LOG(LogMemoryCommandParser, Error, TEXT("Cannot pass a negative value (%.2f) for TimeC"), TimeC);
 					return false;
 				}
 			}
@@ -151,27 +156,21 @@ bool FMemoryExportCommandParser::ParseParameters(const TCHAR* Cmd, FOutputDevice
 				{
 					continue;
 				}
-				if (!Token.IsNumeric())
-				{
-					UE_LOG(LogMemoryCommandParser, Error, TEXT("TimeD value %s is not a valid number"), *Token);
-					return false;
-				}
-				TimeD = FCString::Atod(*Token);
-				if (TimeD < 0.0)
+				if (!ParseTimeParameter(TimeD, Token, TEXT("D")))
 				{
-					UE_LOG(LogMemoryCommandParser, Error, TEXT("Cannot pass a negative value (%.2f) for TimeD"), TimeD);
 					return false;
 				}
 			}
 			else if (Token.StartsWith(MaxResultsToken))
 			{
+				Token.RightChopInline(UE_ARRAY_COUNT(MaxResultsToken) - 1);
+
 				// If user pass the parameter without values or with garbage, Atoi handles it
 				// and MaxResults will default to 0
-				Token.RightChopInline(UE_ARRAY_COUNT(MaxResultsToken) - 1);
 				MaxResults = FCString::Atoi(*Token);
 				if (MaxResults < 0)
 				{
-					UE_LOG(LogMemoryCommandParser, Error, TEXT("Cannot pass a negative value (%d) for MaxResults"), MaxResults);
+					UE_LOG(LogMemoryCommandParser, Error, TEXT("Cannot pass a negative value (%d) for MaxResults!"), MaxResults);
 					return false;
 				}
 			}
@@ -194,23 +193,16 @@ bool FMemoryExportCommandParser::ParseParameters(const TCHAR* Cmd, FOutputDevice
 	return true;
 }
 
+////////////////////////////////////////////////////////////////////////////////////////////////////
+
 bool FMemoryExportCommandParser::ValidateParameters(const FMemoryExporter& Exporter)
 {
-	// Validate required parameters
-	if (Rule.IsEmpty())
+	if (RuleName.IsEmpty())
 	{
 		UE_LOG(LogMemoryCommandParser, Error, TEXT("Rule must be specified. Ex: -Rule=AaBCf"));
 		return false;
 	}
 
-	// TODO: Implement all exporter rules. 
-	// Only AaBCf is implemented
-	if (RuleEnum != QueryRule::AaBCf)
-	{
-		UE_LOG(LogMemoryCommandParser, Error, TEXT("Rule '%s' not implemented yet via command line. Please use the UI."), *Rule);
-		return false;
-	}
-
 	// All rules need OutputPath
 	if (OutputPath.IsEmpty())
 	{
@@ -218,52 +210,93 @@ bool FMemoryExportCommandParser::ValidateParameters(const FMemoryExporter& Expor
 		return false;
 	}
 
-	// Validate Time usages based off of rule groups
-	if (!ValidTimeParameters(RuleEnum, RulesUsingTimeA(), BookmarkA, TimeA, TEXT("A"), Exporter))
+	// Validate BookmarkX and TimeX parameters
+	if (!ValidateTimeParameters(RuleEnum, 1, BookmarkA, TimeA, TEXT("A"), Exporter))
 	{
 		return false;
 	}
-	if (!ValidTimeParameters(RuleEnum, RulesUsingTimeB(), BookmarkB, TimeB, TEXT("B"), Exporter))
+	if (!ValidateTimeParameters(RuleEnum, 2, BookmarkB, TimeB, TEXT("B"), Exporter))
 	{
 		return false;
 	}
-	if (!ValidTimeParameters(RuleEnum, RulesUsingTimeC(), BookmarkC, TimeC, TEXT("C"), Exporter))
+	if (!ValidateTimeParameters(RuleEnum, 3, BookmarkC, TimeC, TEXT("C"), Exporter))
 	{
 		return false;
 	}
-	if (!ValidTimeParameters(RuleEnum, RulesUsingTimeD(), BookmarkD, TimeD, TEXT("D"), Exporter))
+	if (!ValidateTimeParameters(RuleEnum, 4, BookmarkD, TimeD, TEXT("D"), Exporter))
 	{
 		return false;
 	}
+
 	return true;
 }
 
-bool FMemoryExportCommandParser::ValidTimeParameters(QueryRule RuleEnumerator, 
-	const TSet<QueryRule>& RulesUsingTimeX, 
-	const FString& BookmarkX, 
+////////////////////////////////////////////////////////////////////////////////////////////////////
+
+uint32 FMemoryExportCommandParser::GetNumTimeMarkers(EQueryRule Rule)
+{
+	switch (Rule)
+	{
+	case EQueryRule::aAf:
+	case EQueryRule::afA:
+	case EQueryRule::Aaf:
+		return 1;
+
+	case EQueryRule::aAfB:
+	case EQueryRule::AaBf:
+	case EQueryRule::aAfaBf:
+	case EQueryRule::AfB:
+	case EQueryRule::AaB:
+	case EQueryRule::AafB:
+	case EQueryRule::aABf:
+	case EQueryRule::AoB:
+	case EQueryRule::AiB:
+		return 2;
+
+	case EQueryRule::AaBCf:
+	case EQueryRule::AaBfC:
+	case EQueryRule::aABfC:
+		return 3;
+
+	case EQueryRule::AaBCfD:
+		return 4;
+
+	default:
+		return 0;
+	}
+}
+
+////////////////////////////////////////////////////////////////////////////////////////////////////
+
+bool FMemoryExportCommandParser::ValidateTimeParameters(
+	EQueryRule Rule,
+	uint32 TimeMarkerId,
+	const FString& BookmarkX,
 	double& TimeX,
 	const TCHAR* TimeName,
 	const FMemoryExporter& Exporter)
 {
-	if (RulesUsingTimeX.Contains(RuleEnumerator))
+	const uint32 NumTimeMarkers = GetNumTimeMarkers(Rule);
+
+	if (TimeMarkerId <= NumTimeMarkers)
 	{
 		if (BookmarkX.IsEmpty() && TimeX < 0.0)
 		{
-			UE_LOG(LogMemoryCommandParser, Error, 
-				TEXT("Time%s must be specified using either -Bookmark%s=<a bookmark name> or -Time%s=<a real positive number in seconds>"),
+			UE_LOG(LogMemoryCommandParser, Error,
+				TEXT("Time%s must be specified using either -Bookmark%s=\"<a bookmark name>\" or -Time%s=<a real positive number in seconds> parameters!"),
 				TimeName, TimeName, TimeName);
 			return false;
 		}
 		if (!BookmarkX.IsEmpty() && TimeX >= 0.0)
 		{
-			UE_LOG(LogMemoryCommandParser, Error, TEXT("Cannot specify both -Bookmark%s= and -Time%s=. Choose one."),
+			UE_LOG(LogMemoryCommandParser, Error, TEXT("Cannot specify both -Bookmark%s= and -Time%s= parameters! Choose one."),
 				TimeName, TimeName);
 			return false;
 		}
 		// Only when a bookmark has been provided, then lookup the Timestamp
 		if (!BookmarkX.IsEmpty() && !Exporter.FindBookmarkByName(BookmarkX, TimeX))
 		{
-			UE_LOG(LogMemoryCommandParser, Error, TEXT("Bookmark '%s' not found, please re-run trace with valid bookmark name"), *BookmarkX);
+			UE_LOG(LogMemoryCommandParser, Error, TEXT("Bookmark '%s' not found! Please re-run trace with valid bookmark name."), *BookmarkX);
 			return false;
 		}
 	}
@@ -271,75 +304,33 @@ bool FMemoryExportCommandParser::ValidTimeParameters(QueryRule RuleEnumerator,
 	return true;
 }
 
-// Certain rules need specific TimeX (Where X=A,B,C or D).
-// Since only one rule is implemented so far (AaBCf), the following functions
-// will serve for devs to add the rules they find using Time A,B,C or D.
-// As more rules are introduced in AllocationsProvider.h, they should be specified in these groups.
-
-const TSet<FMemoryExportCommandParser::QueryRule> FMemoryExportCommandParser::RulesUsingTimeA()
-{
-	static const TSet<QueryRule> Rules =
-	{
-		QueryRule::aAf, QueryRule::afA, QueryRule::Aaf, QueryRule::aAfB,
-		QueryRule::AaBf, QueryRule::aAfaBf, QueryRule::AfB, QueryRule::AaB,
-		QueryRule::AafB, QueryRule::aABf, QueryRule::AaBCf, QueryRule::AaBfC,
-		QueryRule::aABfC, QueryRule::AaBCfD, QueryRule::AoB, QueryRule::AiB
-	}; // TimeA is used in all rules
-
-	return Rules;
-}
-
-const TSet<FMemoryExportCommandParser::QueryRule> FMemoryExportCommandParser::RulesUsingTimeB()
-{
-	static const TSet<QueryRule> Rules =
-	{
-		QueryRule::aAfB, QueryRule::AaBf, QueryRule::aAfaBf, QueryRule::AfB, QueryRule::AaB,
-		QueryRule::AafB, QueryRule::aABf, QueryRule::AaBCf, QueryRule::AaBfC, QueryRule::aABfC,
-		QueryRule::AaBCfD, QueryRule::AoB, QueryRule::AiB
-	};
-	return Rules;
-}
+////////////////////////////////////////////////////////////////////////////////////////////////////
 
-const TSet<FMemoryExportCommandParser::QueryRule> FMemoryExportCommandParser::RulesUsingTimeC()
+const TMap<FString, FMemoryExportCommandParser::EQueryRule> FMemoryExportCommandParser::QueryRuleMap()
 {
-	static const TSet<QueryRule> Rules =
+	static const TMap<FString, EQueryRule> Map =
 	{
-		QueryRule::AaBCf, QueryRule::AaBfC, QueryRule::aABfC, QueryRule::AaBCfD
-	};
-	return Rules;
-}
-
-const TSet<FMemoryExportCommandParser::QueryRule> FMemoryExportCommandParser::RulesUsingTimeD()
-{
-	static const TSet<QueryRule> Rules =
-	{
-		QueryRule::AaBCfD
-	};
-	return Rules;
-}
-
-const TMap<FString, FMemoryExportCommandParser::QueryRule> FMemoryExportCommandParser::QueryRuleMap()
-{
-	static const TMap<FString, QueryRule> Map =
-	{
-		{ TEXT("aAf"),     QueryRule::aAf },
-		{ TEXT("afA"),     QueryRule::afA },
-		{ TEXT("Aaf"),     QueryRule::Aaf },
-		{ TEXT("aAfB"),    QueryRule::aAfB },
-		{ TEXT("AaBf"),    QueryRule::AaBf },
-		{ TEXT("aAfaBf"),  QueryRule::aAfaBf },
-		{ TEXT("AfB"),     QueryRule::AfB },
-		{ TEXT("AaB"),     QueryRule::AaB },
-		{ TEXT("AafB"),    QueryRule::AafB },
-		{ TEXT("aABf"),    QueryRule::aABf },
-		{ TEXT("AaBCf"),   QueryRule::AaBCf },
-		{ TEXT("AaBfC"),   QueryRule::AaBfC },
-		{ TEXT("aABfC"),   QueryRule::aABfC },
-		{ TEXT("AaBCfD"),  QueryRule::AaBCfD },
-		{ TEXT("AoB"),     QueryRule::AoB },
-		{ TEXT("AiB"),     QueryRule::AiB }
+		{ TEXT("aAf"),     EQueryRule::aAf },
+		{ TEXT("afA"),     EQueryRule::afA },
+		{ TEXT("Aaf"),     EQueryRule::Aaf },
+		{ TEXT("aAfB"),    EQueryRule::aAfB },
+		{ TEXT("AaBf"),    EQueryRule::AaBf },
+		{ TEXT("aAfaBf"),  EQueryRule::aAfaBf },
+		{ TEXT("AfB"),     EQueryRule::AfB },
+		{ TEXT("AaB"),     EQueryRule::AaB },
+		{ TEXT("AafB"),    EQueryRule::AafB },
+		{ TEXT("aABf"),    EQueryRule::aABf },
+		{ TEXT("AaBCf"),   EQueryRule::AaBCf },
+		{ TEXT("AaBfC"),   EQueryRule::AaBfC },
+		{ TEXT("aABfC"),   EQueryRule::aABfC },
+		{ TEXT("AaBCfD"),  EQueryRule::AaBCfD },
+		{ TEXT("AoB"),     EQueryRule::AoB },
+		{ TEXT("AiB"),     EQueryRule::AiB }
 	};
 	return Map;
 }
 
+////////////////////////////////////////////////////////////////////////////////////////////////////
+////////////////////////////////////////////////////////////////////////////////////////////////////
+
 } // namespace UE::Insights::MemoryProfiler
diff --git a/Engine/Source/Developer/TraceInsights/Private/Insights/MemoryProfiler/MemoryExportCommandParser.h b/Engine/Source/Developer/TraceInsights/Private/Insights/MemoryProfiler/MemoryExportCommandParser.h
--- a/Engine/Source/Developer/TraceInsights/Private/Insights/MemoryProfiler/MemoryExportCommandParser.h
+++ b/Engine/Source/Developer/TraceInsights/Private/Insights/MemoryProfiler/MemoryExportCommandParser.h
@@ -8,17 +8,18 @@
 #include "TraceServices/Model/AllocationsProvider.h"
 #include "TraceServices/Model/Memory.h"
 
-// ViewModels
+// TraceInsights
 #include "ViewModels/MemoryExporter.h"
 
-
 namespace UE::Insights::MemoryProfiler
 {
+
 DECLARE_LOG_CATEGORY_EXTERN(LogMemoryCommandParser, Log, All);
 
 class FMemoryExportCommandParser
 {
-typedef TraceServices::IAllocationsProvider::EQueryRule QueryRule;
+public:
+	typedef TraceServices::IAllocationsProvider::EQueryRule EQueryRule;
 
 public:
 	FMemoryExportCommandParser();
@@ -28,9 +29,10 @@ public:
 	bool ValidateParameters(const FMemoryExporter& Exporter);
 
 	// Getters
-	const FString& GetRuleName() const { return Rule; }
-	const QueryRule GetRuleEnum() const { return RuleEnum; }
-	const FString& GetOutputPathName() const { return OutputPath;  }
+	const FString& GetRuleName() const { return RuleName; }
+	uint32 GetRuleNumTimeMarkers() const { return GetNumTimeMarkers(RuleEnum); }
+	const EQueryRule GetRuleEnum() const { return RuleEnum; }
+	const FString& GetOutputPathName() const { return OutputPath; }
 	const FString& GetBookmarkA() const { return BookmarkA; }
 	const FString& GetBookmarkB() const { return BookmarkB; }
 	const FString& GetBookmarkC() const { return BookmarkC; }
@@ -42,32 +44,26 @@ public:
 	int32 GetMaxResults() const { return MaxResults; }
 	const FString& GetColumns() const { return Columns; }
 
-
 private:
-	bool ValidTimeParameters(QueryRule RuleEnumerator,
-		const TSet<QueryRule>& RulesUsingTime,
+	bool ValidateTimeParameters(
+		EQueryRule Rule,
+		uint32 TimeMarkerId,
 		const FString& Bookmark,
 		double& Time,
 		const TCHAR* TimeName,
 		const FMemoryExporter& Exporter);
 
-	static const TSet<QueryRule> RulesUsingTimeA();
-	static const TSet<QueryRule> RulesUsingTimeB();
-	static const TSet<QueryRule> RulesUsingTimeC();
-	static const TSet<QueryRule> RulesUsingTimeD();
-	static const TMap<FString, QueryRule> QueryRuleMap();
+	static uint32 GetNumTimeMarkers(EQueryRule Rule);
+	static const TMap<FString, EQueryRule> QueryRuleMap();
 
 private:
-	FString Rule;
-	QueryRule RuleEnum;
+	FString RuleName;
+	EQueryRule RuleEnum;
 	FString OutputPath;
 	FString BookmarkA, BookmarkB, BookmarkC, BookmarkD;
 	double TimeA, TimeB, TimeC, TimeD;
 	int32 MaxResults;
 	FString Columns;
-
-		
-
 };
 
 } // namespace UE::Insights::MemoryProfiler
\ No newline at end of file
diff --git a/Engine/Source/Developer/TraceInsights/Private/Insights/MemoryProfiler/MemoryProfilerManager.cpp b/Engine/Source/Developer/TraceInsights/Private/Insights/MemoryProfiler/MemoryProfilerManager.cpp
--- a/Engine/Source/Developer/TraceInsights/Private/Insights/MemoryProfiler/MemoryProfilerManager.cpp
+++ b/Engine/Source/Developer/TraceInsights/Private/Insights/MemoryProfiler/MemoryProfilerManager.cpp
@@ -21,14 +21,12 @@
 #include "Insights/Common/InsightsMenuBuilder.h"
 #include "Insights/InsightsManager.h"
 #include "Insights/InsightsStyle.h"
+#include "Insights/MemoryProfiler/MemoryExportCommandParser.h"
+#include "Insights/MemoryProfiler/ViewModels/MemoryExporter.h"
 #include "Insights/MemoryProfiler/ViewModels/MemorySharedState.h"
 #include "Insights/MemoryProfiler/Widgets/SMemoryProfilerWindow.h"
 #include "Insights/Widgets/STimingView.h"
 
-// Commad Line Exporters
-#include "MemoryExportCommandParser.h"
-#include "ViewModels/MemoryExporter.h"
-
 #define LOCTEXT_NAMESPACE "UE::Insights::MemoryProfiler"
 
 // For Memory Insights in editor preview.
@@ -455,63 +453,14 @@ void FMemoryProfilerManager::OnWindowClosedEvent()
 
 ////////////////////////////////////////////////////////////////////////////////////////////////////
 
-typedef TraceServices::IAllocationsProvider::EQueryRule QueryRule;
-static const TSet<QueryRule> RulesUsingTimeA =
-{
-	QueryRule::aAf, QueryRule::afA, QueryRule::Aaf, QueryRule::aAfB,
-	QueryRule::AaBf, QueryRule::aAfaBf, QueryRule::AfB, QueryRule::AaB,
-	QueryRule::AafB, QueryRule::aABf, QueryRule::AaBCf, QueryRule::AaBfC,
-	QueryRule::aABfC, QueryRule::AaBCfD, QueryRule::AoB, QueryRule::AiB
-}; // TimeA is used in all rules
-static const TSet<QueryRule> RulesUsingTimeB =
-{
-	QueryRule::aAfB, QueryRule::AaBf, QueryRule::aAfaBf, QueryRule::AfB, QueryRule::AaB,
-	QueryRule::AafB, QueryRule::aABf, QueryRule::AaBCf, QueryRule::AaBfC, QueryRule::aABfC,
-	QueryRule::AaBCfD, QueryRule::AoB, QueryRule::AiB
-};
-static const TSet<QueryRule> RulesUsingTimeC =
-{
-	QueryRule::AaBCf, QueryRule::AaBfC, QueryRule::aABfC, QueryRule::AaBCfD
-};
-static const TSet<QueryRule> RulesUsingTimeD =
-{
-	QueryRule::AaBCfD
-};
-static const TMap<FString, QueryRule> QueryRuleMap =
-{
-	{ TEXT("aAf"),     QueryRule::aAf },
-	{ TEXT("afA"),     QueryRule::afA },
-	{ TEXT("Aaf"),     QueryRule::Aaf },
-	{ TEXT("aAfB"),    QueryRule::aAfB },
-	{ TEXT("AaBf"),    QueryRule::AaBf },
-	{ TEXT("aAfaBf"),  QueryRule::aAfaBf },
-	{ TEXT("AfB"),     QueryRule::AfB },
-	{ TEXT("AaB"),     QueryRule::AaB },
-	{ TEXT("AafB"),    QueryRule::AafB },
-	{ TEXT("aABf"),    QueryRule::aABf },
-	{ TEXT("AaBCf"),   QueryRule::AaBCf },
-	{ TEXT("AaBfC"),   QueryRule::AaBfC },
-	{ TEXT("aABfC"),   QueryRule::aABfC },
-	{ TEXT("AaBCfD"),  QueryRule::AaBCfD },
-	{ TEXT("AoB"),     QueryRule::AoB },
-	{ TEXT("AiB"),     QueryRule::AiB }
-};
-
-static bool ValidateTimeParameters(QueryRule RuleEnum,
-	const TSet<QueryRule>& RulesUsingTime,
-	const FString& Bookmark,
-	double& Time,
-	const TCHAR* TimeName,
-	const FMemoryExporter& Exporter)
-{
-	return false;
-}
-
 bool FMemoryProfilerManager::Exec(const TCHAR* Cmd, FOutputDevice& Ar)
 {
-	// Parameterized Memory Leak Export Command
+	// Export allocations based on memory investigation rules
 	if (FParse::Command(&Cmd, TEXT("MemoryInsights.ExportAllocs")))
 	{
+		FStopwatch Stopwatch;
+		Stopwatch.Start();
+
 		TSharedPtr<FInsightsManager> InsightsManager = FInsightsManager::Get();
 		if (!InsightsManager.IsValid())
 		{
@@ -528,6 +477,7 @@ bool FMemoryProfilerManager::Exec(const TCHAR* Cmd, FOutputDevice& Ar)
 		FMemoryExporter Exporter(*InsightsManager->GetSession().Get());
 
 		FMemoryExportCommandParser ExportCommandParser;
+
 		if (!ExportCommandParser.ParseParameters(Cmd, Ar))
 		{
 			UE_LOG(LogMemoryProfiler, Error, TEXT("Exporter stopped while parsing cmd parameters"));
@@ -540,48 +490,65 @@ bool FMemoryProfilerManager::Exec(const TCHAR* Cmd, FOutputDevice& Ar)
 			return false;
 		}
 
-		// Below logic corresponds to Rule=AaBCf
-		if(ExportCommandParser.GetRuleEnum() == TraceServices::IAllocationsProvider::EQueryRule::AaBCf)
+		// All QueryRule are supported
+		UE_LOG(LogMemoryProfiler, Log, TEXT(" === MemoryInsights.ExportAllocs -Rule=\"%s\" === "), *ExportCommandParser.GetRuleName());
+
+		FMemoryExporter::FExportMemoryAllocsParams Params;
+		Params.Rule = ExportCommandParser.GetRuleEnum();
+		Params.OutputPathName = ExportCommandParser.GetOutputPathName();
+		Params.TimeA = ExportCommandParser.GetTimeA();
+		Params.TimeB = ExportCommandParser.GetTimeB();
+		Params.TimeC = ExportCommandParser.GetTimeC();
+		Params.TimeD = ExportCommandParser.GetTimeD();
+		Params.MaxResults = ExportCommandParser.GetMaxResults();
+		Params.Columns = ExportCommandParser.GetColumns();
+
+		UE_LOG(LogMemoryProfiler, Log, TEXT("Parameters provided by user:"));
+		UE_LOG(LogMemoryProfiler, Log, TEXT("-> Rule:       %s"), *ExportCommandParser.GetRuleName());
+		UE_LOG(LogMemoryProfiler, Log, TEXT("-> Output:     %s"), *Params.OutputPathName);
+
+		const uint32 RuleNumTimeMarkers = ExportCommandParser.GetRuleNumTimeMarkers();
+		if (RuleNumTimeMarkers >= 1)
 		{
-			UE_LOG(LogMemoryProfiler, Log, TEXT(" === Memory Leak Exporter. Rule: \"%s\" === "), *ExportCommandParser.GetRuleName());
-			FMemoryExporter::FExportMemoryLeakSnapshotParams SnapshotParams;
-			SnapshotParams.OutputPathName = ExportCommandParser.GetOutputPathName();
-			SnapshotParams.TimeA = ExportCommandParser.GetTimeA();
-			SnapshotParams.TimeB = ExportCommandParser.GetTimeB();
-			SnapshotParams.TimeC = ExportCommandParser.GetTimeC();
-			SnapshotParams.MaxResults = ExportCommandParser.GetMaxResults();
-			SnapshotParams.Columns = ExportCommandParser.GetColumns();
-
 			FString BookmarkA = ExportCommandParser.GetBookmarkA();
+			UE_LOG(LogMemoryProfiler, Log, TEXT("-> TimeA:      %.2f%s"),
+				Params.TimeA, BookmarkA.IsEmpty() ? TEXT("") : *FString::Printf(TEXT(" (from bookmark '%s')"), *BookmarkA));
+		}
+		if (RuleNumTimeMarkers >= 2)
+		{
 			FString BookmarkB = ExportCommandParser.GetBookmarkB();
+			UE_LOG(LogMemoryProfiler, Log, TEXT("-> TimeB:      %.2f%s"),
+				Params.TimeB, BookmarkB.IsEmpty() ? TEXT("") : *FString::Printf(TEXT(" (from bookmark '%s')"), *BookmarkB));
+		}
+		if (RuleNumTimeMarkers >= 3)
+		{
 			FString BookmarkC = ExportCommandParser.GetBookmarkC();
+			UE_LOG(LogMemoryProfiler, Log, TEXT("-> TimeC:      %.2f%s"),
+				Params.TimeC, BookmarkC.IsEmpty() ? TEXT("") : *FString::Printf(TEXT(" (from bookmark '%s')"), *BookmarkC));
+		}
+		if (RuleNumTimeMarkers >= 4)
+		{
+			FString BookmarkD = ExportCommandParser.GetBookmarkD();
+			UE_LOG(LogMemoryProfiler, Log, TEXT("-> TimeD:      %.2f%s"),
+				Params.TimeD, BookmarkD.IsEmpty() ? TEXT("") : *FString::Printf(TEXT(" (from bookmark '%s')"), *BookmarkD));
+		}
 
-			UE_LOG(LogMemoryProfiler, Log, TEXT("Parameters provided by user:"));
-			UE_LOG(LogMemoryProfiler, Log, TEXT("-> Output:     %s"), *SnapshotParams.OutputPathName);
-			UE_LOG(LogMemoryProfiler, Log, TEXT("-> TimeA:      %.2f%s"), SnapshotParams.TimeA, BookmarkA.IsEmpty() ? TEXT("") : *FString::Printf(TEXT(" (from bookmark '%s')"), *BookmarkA));
-			UE_LOG(LogMemoryProfiler, Log, TEXT("-> TimeB:      %.2f%s"), SnapshotParams.TimeB, BookmarkB.IsEmpty() ? TEXT("") : *FString::Printf(TEXT(" (from bookmark '%s')"), *BookmarkB));
-			UE_LOG(LogMemoryProfiler, Log, TEXT("-> TimeC:      %.2f%s"), SnapshotParams.TimeC, BookmarkC.IsEmpty() ? TEXT("") : *FString::Printf(TEXT(" (from bookmark '%s')"), *BookmarkC));
-			if (ExportCommandParser.GetTimeD() >= 0.0 || !ExportCommandParser.GetBookmarkD().IsEmpty())
-			{
-				UE_LOG(LogMemoryProfiler, Warning, TEXT("-> TimeD SKIPPED. Not required in this rule"));
-			}
-			UE_LOG(LogMemoryProfiler, Log, TEXT("-> MaxResults: %d (0 = unlimited)"), SnapshotParams.MaxResults);
-			UE_LOG(LogMemoryProfiler, Log, TEXT("Starting export..."));
-			int ResultCount = Exporter.ExportMemoryLeakSnapshotAsText(SnapshotParams);
+		UE_LOG(LogMemoryProfiler, Log, TEXT("-> MaxResults: %d (0 = unlimited)"), Params.MaxResults);
 
-			if (ResultCount >= 0)
-			{
-				UE_LOG(LogMemoryProfiler, Log, TEXT("SUCCESS! Exported %d allocations"), ResultCount);
-				UE_LOG(LogMemoryProfiler, Log, TEXT("Output file: %s"), *SnapshotParams.OutputPathName);
-				return true;
-			}
-			else
-			{
-				UE_LOG(LogMemoryProfiler, Error, TEXT("Export failed with code %d"), ResultCount);
-				return false;
-			}
+		UE_LOG(LogMemoryProfiler, Log, TEXT("Starting export..."));
+		int ResultCount = Exporter.ExportMemoryAllocsAsText(Params);
+
+		if (ResultCount >= 0)
+		{
+			UE_LOG(LogMemoryProfiler, Log, TEXT("SUCCESS! Exported %d allocations"), ResultCount);
+			UE_LOG(LogMemoryProfiler, Log, TEXT("Output file: %s"), *Params.OutputPathName);
+			return true;
+		}
+		else
+		{
+			UE_LOG(LogMemoryProfiler, Error, TEXT("Export failed with code %d"), ResultCount);
+			return false;
 		}
-		
 	}
 	return false;
 }
diff --git a/Engine/Source/Developer/TraceInsights/Private/Insights/MemoryProfiler/ViewModels/MemoryExporter.cpp b/Engine/Source/Developer/TraceInsights/Private/Insights/MemoryProfiler/ViewModels/MemoryExporter.cpp
--- a/Engine/Source/Developer/TraceInsights/Private/Insights/MemoryProfiler/ViewModels/MemoryExporter.cpp
+++ b/Engine/Source/Developer/TraceInsights/Private/Insights/MemoryProfiler/ViewModels/MemoryExporter.cpp
@@ -4,537 +4,627 @@
 
 #include "Containers/Map.h"
 #include "Containers/StringConv.h"
+#include "Containers/Utf8String.h"
 #include "GenericPlatform/GenericPlatformFile.h"
 #include "HAL/PlatformFileManager.h"
 #include "HAL/PlatformProcess.h"
-#include "Misc/Paths.h"
 #include "Logging/LogMacros.h"
+#include "Misc/Paths.h"
 
 // TraceServices
 #include "Common/ProviderLock.h"
-#include "TraceServices/Model/AllocationsProvider.h"
+#include "TraceServices/Model/AnalysisSession.h"
 #include "TraceServices/Model/Bookmarks.h"
 #include "TraceServices/Model/Callstack.h"
 #include "TraceServices/Model/Definitions.h"
-#include "TraceServices/Model/MetadataProvider.h"
 #include "TraceServices/Model/Strings.h"
 
+// TraceInsightsCore
+#include "InsightsCore/Common/Stopwatch.h"
+
+// TraceInsights
+#include "Insights/Common/Utf8FileWriter.h"
+#include "Insights/MemoryProfiler/ViewModels/MemAllocTable.h"
 
-DEFINE_LOG_CATEGORY_STATIC(LogMemoryLeakSnapshot, Log, All);
+#include <limits>
+
+DEFINE_LOG_CATEGORY_STATIC(LogMemoryExporter, Log, All);
 
 namespace UE::Insights::MemoryProfiler
 {
-	///////////////////////////////////////////////////////////////////////////////////////////////////
-	// FMemoryExporter
-	////////////////////////////////////////////////////////////////////////////////////////////////////
 
-	FMemoryExporter::FMemoryExporter(const TraceServices::IAnalysisSession& InSession)
-		: Session(InSession)
+////////////////////////////////////////////////////////////////////////////////////////////////////
+// FMemoryExporter
+////////////////////////////////////////////////////////////////////////////////////////////////////
+
+const FName FMemoryExporter::AllocCallstackColumnId("AllocCallstack");
+const FName FMemoryExporter::FreeCallstackColumnId("FreeCallstack");
+
+////////////////////////////////////////////////////////////////////////////////////////////////////
+
+FMemoryExporter::FMemoryExporter(const TraceServices::IAnalysisSession& InSession)
+	: Session(InSession)
+{
+}
+
+////////////////////////////////////////////////////////////////////////////////////////////////////
+
+FMemoryExporter::~FMemoryExporter()
+{
+}
+
+////////////////////////////////////////////////////////////////////////////////////////////////////
+
+void FMemoryExporter::InitAvailableColumns()
+{
+	if (!AvailableColumns.IsEmpty())
 	{
+		return;
 	}
 
-	////////////////////////////////////////////////////////////////////////////////////////////////////
+	AvailableColumns.Add(FMemAllocTableColumns::StartEventIndexColumnId);
+	AvailableColumns.Add(FMemAllocTableColumns::EndEventIndexColumnId);
+	AvailableColumns.Add(FMemAllocTableColumns::EventDistanceColumnId);
+	AvailableColumns.Add(FMemAllocTableColumns::StartTimeColumnId);
+	AvailableColumns.Add(FMemAllocTableColumns::EndTimeColumnId);
+	AvailableColumns.Add(FMemAllocTableColumns::DurationColumnId);
+	AvailableColumns.Add(FMemAllocTableColumns::AddressColumnId);
+	AvailableColumns.Add(FMemAllocTableColumns::MemoryPageColumnId);
+	AvailableColumns.Add(FMemAllocTableColumns::SizeColumnId);
+	AvailableColumns.Add(FMemAllocTableColumns::TagColumnId);
+	AvailableColumns.Add(FMemAllocTableColumns::AssetColumnId);
+	AvailableColumns.Add(FMemAllocTableColumns::ClassNameColumnId);
+	AvailableColumns.Add(FMemAllocTableColumns::PackageColumnId);
+	AvailableColumns.Add(FMemAllocTableColumns::AllocThreadColumnId);
+	AvailableColumns.Add(FMemAllocTableColumns::FreeThreadColumnId);
+	AvailableColumns.Add(FMemAllocTableColumns::AllocFunctionColumnId);
+	AvailableColumns.Add(FMemAllocTableColumns::FreeFunctionColumnId);
+	AvailableColumns.Add(FMemAllocTableColumns::AllocSourceFileColumnId);
+	AvailableColumns.Add(FMemAllocTableColumns::FreeSourceFileColumnId);
+	AvailableColumns.Add(FMemoryExporter::AllocCallstackColumnId);
+	AvailableColumns.Add(FMemoryExporter::FreeCallstackColumnId);
+}
+
+////////////////////////////////////////////////////////////////////////////////////////////////////
+
+void FMemoryExporter::InitDefaultColumns()
+{
+	InitAvailableColumns();
 
-	FMemoryExporter::~FMemoryExporter()
+	if (DefaultColumns.IsEmpty())
 	{
+		DefaultColumns = AvailableColumns;
 	}
+}
 
-	////////////////////////////////////////////////////////////////////////////////////////////////////
+////////////////////////////////////////////////////////////////////////////////////////////////////
 
-	const FName FMemoryExporter::ExportMemoryLeakSnapshot_AllocCallstack("AllocCallstack");
-	const FName FMemoryExporter::ExportMemoryLeakSnapshot_FreeCallstack("FreeCallstack");
-	
-	////////////////////////////////////////////////////////////////////////////////////////////////////
+void FMemoryExporter::MakeColumnList(const FString& InColumnsString, TArray<FName>& OutColumnList)
+{
+	InitAvailableColumns();
+
+	TArray<FString> Columns;
+	InColumnsString.ParseIntoArray(Columns, TEXT(","), true);
 
-	void FMemoryExporter::ExportMemorySnapshot_InitColumns()
+	// Provide some hints of malformed list of columns (the white space case)
+	if (InColumnsString.StartsWith(TEXT(",")) || InColumnsString.EndsWith(TEXT(",")))
 	{
-		InitializeAvailableColumns();
-		if (SnapshotDefaultColumns.IsEmpty())
-		{
-			SnapshotDefaultColumns = AvailableColumns;
-		}
+		UE_LOG(LogMemoryExporter, Warning, TEXT("Column list has leading/trailing comma: '%s'"), *InColumnsString);
+		UE_LOG(LogMemoryExporter, Warning,
+			TEXT("Check you didn't separate the list of columns with white spaces or remove extra commas."));
 	}
 
-	void FMemoryExporter::InitializeAvailableColumns()
+	for (const FString& ColumnWildcard : Columns)
 	{
-		if (!AvailableColumns.IsEmpty())
+		FName ColumnName(*ColumnWildcard);
+		if (AvailableColumns.Contains(ColumnName))
 		{
-			return;
+			OutColumnList.Add(ColumnName);
 		}
-
-		AvailableColumns.Add(FMemAllocTableColumns::StartEventIndexColumnId);
-		AvailableColumns.Add(FMemAllocTableColumns::EndEventIndexColumnId);
-		AvailableColumns.Add(FMemAllocTableColumns::EventDistanceColumnId);
-		AvailableColumns.Add(FMemAllocTableColumns::StartTimeColumnId);
-		AvailableColumns.Add(FMemAllocTableColumns::EndTimeColumnId);
-		AvailableColumns.Add(FMemAllocTableColumns::DurationColumnId);
-		AvailableColumns.Add(FMemAllocTableColumns::AddressColumnId);
-		AvailableColumns.Add(FMemAllocTableColumns::MemoryPageColumnId);
-		AvailableColumns.Add(FMemAllocTableColumns::SizeColumnId);
-		AvailableColumns.Add(FMemAllocTableColumns::TagColumnId);
-		AvailableColumns.Add(FMemAllocTableColumns::AssetColumnId);
-		AvailableColumns.Add(FMemAllocTableColumns::ClassNameColumnId);
-		AvailableColumns.Add(FMemAllocTableColumns::PackageColumnId);
-		AvailableColumns.Add(FMemAllocTableColumns::AllocThreadColumnId);
-		AvailableColumns.Add(FMemAllocTableColumns::FreeThreadColumnId);
-		AvailableColumns.Add(FMemAllocTableColumns::AllocFunctionColumnId);
-		AvailableColumns.Add(FMemAllocTableColumns::FreeFunctionColumnId);
-		AvailableColumns.Add(FMemAllocTableColumns::AllocSourceFileColumnId);
-		AvailableColumns.Add(FMemAllocTableColumns::FreeSourceFileColumnId);
-		AvailableColumns.Add(ExportMemoryLeakSnapshot_AllocCallstack);
-		AvailableColumns.Add(ExportMemoryLeakSnapshot_FreeCallstack);
-	}
-
-	////////////////////////////////////////////////////////////////////////////////////////////////////
-	// Copying the way TimingExporter generate the columns. 
-	// We need this function now to support handling specific columns
-	void FMemoryExporter::MakeExportMemoryLeakSnapshotColumnList(const FString& InColumnsString, TArray<FName>& OutColumnList)
-	{
-		InitializeAvailableColumns();
-
-		TArray<FString> Columns;
-		InColumnsString.ParseIntoArray(Columns, TEXT(","), true);
-
-		// Provide some hints of malformed list of columns (the white space case)
-		if (InColumnsString.StartsWith(TEXT(",")) || InColumnsString.EndsWith(TEXT(",")))
-		{
-			UE_LOG(LogMemoryLeakSnapshot, Warning, TEXT("Column list has leading/trailing comma: '%s'"), *InColumnsString);
-			UE_LOG(LogMemoryLeakSnapshot, Warning, 
-				TEXT("Check you didn't separate the list of columns with white spaces or remove extra commas."));
-		}
-
-		for (const FString& ColumnWildcard : Columns)
+		else
 		{
-			FName ColumnName(*ColumnWildcard);
-			if (AvailableColumns.Contains(ColumnName))
-			{
-				OutColumnList.Add(ColumnName);
-			}
-			else
+			for (const FName& Column : AvailableColumns)
 			{
-				for (const FName& Column : AvailableColumns)
+				if (Column.GetPlainNameString().MatchesWildcard(ColumnWildcard))
 				{
-					if (Column.GetPlainNameString().MatchesWildcard(ColumnWildcard))
-					{
-						OutColumnList.Add(Column);
-					}
+					OutColumnList.Add(Column);
 				}
 			}
 		}
 	}
+}
 
-	////////////////////////////////////////////////////////////////////////////////////////////////////
+////////////////////////////////////////////////////////////////////////////////////////////////////
 
-	int32 FMemoryExporter::ExportMemoryLeakSnapshotAsText(FExportMemoryLeakSnapshotParams& Params)
+int32 FMemoryExporter::ExportMemoryAllocsAsText(FExportMemoryAllocsParams& Params)
+{
+	UE_LOG(LogMemoryExporter, Log, TEXT("Memory Exporter Started"));
+	UE_LOG(LogMemoryExporter, Log, TEXT("Rule: %s"), LexToString(Params.Rule));
+	UE_LOG(LogMemoryExporter, Log, TEXT("Output CSV file will be saved at: %s"), *Params.OutputPathName);
+	UE_LOG(LogMemoryExporter, Log, TEXT("Time Markers: A=%.2f, B=%.2f, C=%.2f, D=%.2f"),
+		Params.TimeA, Params.TimeB, Params.TimeC, Params.TimeD);
+
+	UE_LOG(LogMemoryExporter, Log, TEXT("Querying allocations..."));
+	TArray<const TraceServices::IAllocationsProvider::FAllocation*> Allocations;
+	if (!QueryAllocations(Params, Allocations))
 	{
-		UE_LOG(LogMemoryLeakSnapshot, Log, TEXT("Memory Leak Exporter Started"));
-		UE_LOG(LogMemoryLeakSnapshot, Log, TEXT("Output CSV file will be saved at: %s"), *Params.OutputPathName);
-		UE_LOG(LogMemoryLeakSnapshot, Log, TEXT("Memory Leak Analysis Times: A=%.2f, B=%.2f, C=%.2f"), Params.TimeA, Params.TimeB, Params.TimeC);
+		UE_LOG(LogMemoryExporter, Error, TEXT("Failed to query allocations"));
+		return -1;
+	}
+	UE_LOG(LogMemoryExporter, Log, TEXT("Found %d allocations"), Allocations.Num());
 
-		UE_LOG(LogMemoryLeakSnapshot, Log, TEXT("Querying allocations..."));
-		TArray<const TraceServices::IAllocationsProvider::FAllocation*> Allocations;
-		if (!QueryAllocations(Params, Allocations))
-		{
-			UE_LOG(LogMemoryLeakSnapshot, Error, TEXT("Failed to query allocations"));
-			return -1;
-		}
-		UE_LOG(LogMemoryLeakSnapshot, Log, TEXT("Found %d allocations"), Allocations.Num());
+	// Determine what columns to export
+	TArray<FName> ExportColumns;
+	if (Params.Columns.IsEmpty())
+	{
+		InitDefaultColumns();
+		ExportColumns = DefaultColumns;
+		UE_LOG(LogMemoryExporter, Log, TEXT("Using all %d columns available by default"), ExportColumns.Num());
+	}
+	else
+	{
+		MakeColumnList(Params.Columns, ExportColumns);
+		UE_LOG(LogMemoryExporter, Log, TEXT("Using the %d columns passed in -Columns: %s"),
+			ExportColumns.Num(), *Params.Columns);
+	}
 
-		// Determine what columns to export
-		TArray<FName> ExportColumns;
-		if (Params.Columns.IsEmpty())
-		{
-			ExportMemorySnapshot_InitColumns();
-			ExportColumns = SnapshotDefaultColumns;
-			UE_LOG(LogMemoryLeakSnapshot, Log, TEXT("Using all %d columns available by default"), ExportColumns.Num());
-		}
-		else
-		{
-			MakeExportMemoryLeakSnapshotColumnList(Params.Columns, ExportColumns);
-			UE_LOG(LogMemoryLeakSnapshot, Log, TEXT("Using the %d columns passed in -Columns: %s"),
-				ExportColumns.Num(), *Params.Columns);
-		}
-		
-		// Adding limits just in case user needs for investigation
-		if (Params.MaxResults > 0 && Allocations.Num() > Params.MaxResults)
-		{
-			Allocations.SetNum(Params.MaxResults);
-			UE_LOG(LogMemoryLeakSnapshot, Log, TEXT("Export limited by the user to %d allocations"), Params.MaxResults);
-		}
+	// Adding limits just in case user needs for investigation
+	if (Params.MaxResults > 0 && Allocations.Num() > Params.MaxResults)
+	{
+		Allocations.SetNum(Params.MaxResults);
+		UE_LOG(LogMemoryExporter, Log, TEXT("Export limited by the user to %d allocations"), Params.MaxResults);
+	}
 
-		// Write the CSV export
-		if (!WriteSnapshotCSV(Params.OutputPathName, Allocations, ExportColumns))
-		{
-			UE_LOG(LogMemoryLeakSnapshot, Error, TEXT("Failed to write CSV"));
-			return -1;
-		}
+	// Write the CSV/TSV file
+	if (!ExportAllocations(Params.OutputPathName, Allocations, ExportColumns))
+	{
+		UE_LOG(LogMemoryExporter, Error, TEXT("Failed to write CSV"));
+		return -1;
+	}
 
+	return Allocations.Num();
+}
 
-		return Allocations.Num();
-	}
+////////////////////////////////////////////////////////////////////////////////////////////////////
 
-	////////////////////////////////////////////////////////////////////////////////////////////////////
+bool FMemoryExporter::QueryAllocations(const FExportMemoryAllocsParams& Params, TArray<const TraceServices::IAllocationsProvider::FAllocation*>& OutAllocations)
+{
+	// Initialize required reading scope for sessions and providers
 
-	bool FMemoryExporter::QueryAllocations(const FExportMemoryLeakSnapshotParams& Params, TArray<const TraceServices::IAllocationsProvider::FAllocation*>& OutAllocations)
+	TraceServices::FAnalysisSessionReadScope SessionReadScope(Session);
+	const TraceServices::IAllocationsProvider* AllocationsProvider = TraceServices::ReadAllocationsProvider(Session);
+	if (!AllocationsProvider)
+	{
+		UE_LOG(LogMemoryExporter, Error, TEXT("Nullptr: AllocationsProvider not available"));
+		return false;
+	}
+	TraceServices::FProviderReadScopeLock ProviderReadLock(*AllocationsProvider);
+	if (!AllocationsProvider->IsInitialized())
 	{
-		// Initialize required reading scope for sessions and providers
+		UE_LOG(LogMemoryExporter, Error, TEXT("AllocationsProvider exists but not initialized"));
+		return false;
+	}
 
-		TraceServices::FAnalysisSessionReadScope SessionReadScope(Session);
-		const TraceServices::IAllocationsProvider* AllocationsProvider = TraceServices::ReadAllocationsProvider(Session);
-		if (!AllocationsProvider)
-		{
-			UE_LOG(LogMemoryLeakSnapshot, Error, TEXT("Nullptr: AllocationsProvider not available"));
-			return false;
-		}
-		TraceServices::FProviderReadScopeLock ProviderReadLock(*AllocationsProvider);
-		if (!AllocationsProvider->IsInitialized())
-		{
-			UE_LOG(LogMemoryLeakSnapshot, Error, TEXT("AllocationsProvider exists but not initialized"));
-			return false;
-		}
+	// Setup the query
+	TraceServices::IAllocationsProvider::FQueryParams QueryParams;
+	QueryParams.Rule = Params.Rule;
+	QueryParams.TimeA = Params.TimeA;
+	QueryParams.TimeB = Params.TimeB;
+	QueryParams.TimeC = Params.TimeC;
+	QueryParams.TimeD = Params.TimeD;;
+
+	// Start the query and poll results. The API is asynchronous and we have to collect results in batches from it's queue.
+	TraceServices::IAllocationsProvider::FQueryHandle QueryHandle = AllocationsProvider->StartQuery(QueryParams);
 
-		// Setup the query for memory leaks (AaBCf: It's one of the options from the Investigation tab in the UI)
-		TraceServices::IAllocationsProvider::FQueryParams QueryParams;
-		QueryParams.Rule = TraceServices::IAllocationsProvider::EQueryRule::AaBCf;
-		QueryParams.TimeA = Params.TimeA;
-		QueryParams.TimeB = Params.TimeB;
-		QueryParams.TimeC = Params.TimeC;
-		QueryParams.TimeD = 0.0;
-
-		// Start the query and poll results. The API is asynchronous and we have to collect results in batches from it's queue.
-		TraceServices::IAllocationsProvider::FQueryHandle QueryHandle = AllocationsProvider->StartQuery(QueryParams);
-		
-		TraceServices::IAllocationsProvider::FQueryStatus Status;
-		int BatchCount = 0;
-		do
+	TraceServices::IAllocationsProvider::FQueryStatus Status;
+	int BatchCount = 0;
+	do
+	{
+		Status = AllocationsProvider->PollQuery(QueryHandle);
+		if (Status.Status == TraceServices::IAllocationsProvider::EQueryStatus::Available)
 		{
-			Status = AllocationsProvider->PollQuery(QueryHandle);
-			if (Status.Status == TraceServices::IAllocationsProvider::EQueryStatus::Available)
+			TraceServices::IAllocationsProvider::FQueryResult Result = Status.NextResult();
+			if (Result)
 			{
-				TraceServices::IAllocationsProvider::FQueryResult Result = Status.NextResult();
-				if (Result)
+				const int NumAllocations = Result->Num();
+				UE_LOG(LogMemoryExporter, Log, TEXT("Batch %d: %d allocations"), BatchCount, NumAllocations);
+				for (int i = 0; i < NumAllocations; ++i)
 				{
-					const int NumAllocations = Result->Num();
-					UE_LOG(LogMemoryLeakSnapshot, Log, TEXT("Batch %d: %d allocations"), BatchCount, NumAllocations);
-					for (int i = 0; i < NumAllocations; ++i)
-					{
-						OutAllocations.Add(Result->Get(i));
-					}
+					OutAllocations.Add(Result->Get(i));
 				}
 			}
+		}
 
-			++BatchCount;
+		++BatchCount;
 
-			// Handle time out with an arbitrarily big number for BatchCount
-			if (BatchCount > 10000)
-			{
-				UE_LOG(LogMemoryLeakSnapshot, Error, TEXT("Query timeout after %d batches iterated"), BatchCount);
-				AllocationsProvider->CancelQuery(QueryHandle);
-				return false;
-			}
+		// Handle time out with an arbitrarily big number for BatchCount
+		if (BatchCount > 10000)
+		{
+			UE_LOG(LogMemoryExporter, Error, TEXT("Query timeout after %d batches iterated"), BatchCount);
+			AllocationsProvider->CancelQuery(QueryHandle);
+			return false;
+		}
 
-			// To save some cpu
-			FPlatformProcess::Sleep(0.001f);
-		} while (Status.Status != TraceServices::IAllocationsProvider::EQueryStatus::Done);
+		// To save some cpu
+		FPlatformProcess::Sleep(0.001f);
+	} while (Status.Status != TraceServices::IAllocationsProvider::EQueryStatus::Done);
 
+	return true;
+}
 
-		return true;
-	}
+////////////////////////////////////////////////////////////////////////////////////////////////////
+// This function appends in buffer in real time, because query may retrieve millions of allocations
+// TODO: Evaluate using a Function table to avoid so many if/else if per row
+
+void FMemoryExporter::WriteColumnValue(
+	FUtf8FileWriter& Writer,
+	const FName& ColumnId,
+	const TraceServices::IAllocationsProvider::FAllocation* Alloc,
+	const FMemoryExporter::FWriteColumnValueParams& Params)
+{
+	constexpr uint32 InfiniteEventIndex = std::numeric_limits<uint32>::max();
 
-	////////////////////////////////////////////////////////////////////////////////////////////////////
-	// This function appends in buffer in real time, because query may retrieve millions of allocations
-	// TODO: Evaluate using a Function table to avoid so many if/else if per row
-	void FMemoryExporter::WriteColumnValue(FUtf16Writer& Writer, const FName& ColumnId, 
-		const TraceServices::IAllocationsProvider::FAllocation* Alloc, 
-		const TraceServices::IAllocationsProvider* AllocationsProvider, 
-		const TraceServices::FCallstack* AllocCallstack, 
-		const TraceServices::FCallstack* FreeCallstack, 
-		const FString& AssetName, const FString& ClassName, const FString& PackageName, uint64 PageSize)
+	if (ColumnId == FMemAllocTableColumns::StartEventIndexColumnId)
 	{
-		const uint32 MaxEventIndex = TNumericLimits<uint32>::Max();
-		if (ColumnId == FMemAllocTableColumns::StartEventIndexColumnId)
-		{
-			Writer.Append(LexToString(Alloc->GetStartEventIndex()));
-		}
-		else if (ColumnId == FMemAllocTableColumns::EndEventIndexColumnId)
+		Writer.AppendValue(Alloc->GetStartEventIndex());
+	}
+	else if (ColumnId == FMemAllocTableColumns::EndEventIndexColumnId)
+	{
+		if (Alloc->GetEndEventIndex() == InfiniteEventIndex)
 		{
-			Writer.Append(LexToString(Alloc->GetEndEventIndex()));
+			Writer.AppendValue(UTF8TEXTVIEW("inf"));
 		}
-		else if (ColumnId == FMemAllocTableColumns::EventDistanceColumnId)
-		{
-			const uint32 StartEventIndex = Alloc->GetStartEventIndex();
-			const uint32 EndEventIndex = Alloc->GetEndEventIndex();
-			const uint32 EventDistance = (EndEventIndex == MaxEventIndex) ? MaxEventIndex :
-				(EndEventIndex - StartEventIndex);
-			Writer.Append(LexToString(EventDistance));
-		}
-		else if (ColumnId == FMemAllocTableColumns::StartTimeColumnId)
+		else
 		{
-			Writer.Append(LexToString(Alloc->GetStartTime()));
+			Writer.AppendValue(Alloc->GetEndEventIndex());
 		}
-		else if (ColumnId == FMemAllocTableColumns::EndTimeColumnId)
+	}
+	else if (ColumnId == FMemAllocTableColumns::EventDistanceColumnId)
+	{
+		if (Alloc->GetEndEventIndex() == InfiniteEventIndex)
 		{
-			const uint32 EndEventIndex = Alloc->GetEndEventIndex();
-			if (EndEventIndex == MaxEventIndex)
-			{
-				Writer.AppendString(TEXT("inf"));
-			}
-			else
-			{
-				Writer.Append(LexToString(Alloc->GetEndTime()));
-			}
+			Writer.AppendValue(UTF8TEXTVIEW("inf"));
 		}
-		else if (ColumnId == FMemAllocTableColumns::DurationColumnId)
+		else
 		{
+			const uint32 StartEventIndex = Alloc->GetStartEventIndex();
 			const uint32 EndEventIndex = Alloc->GetEndEventIndex();
-			if (EndEventIndex == MaxEventIndex)
-			{
-				Writer.AppendString(TEXT("inf"));
-			}
-			else
-			{
-				Writer.Append(LexToString(Alloc->GetEndTime() - Alloc->GetStartTime()));
-			}
-		}
-		else if (ColumnId == FMemAllocTableColumns::AddressColumnId)
-		{
-			Writer.Append(LexToString(Alloc->GetAddress()));
-		}
-		else if (ColumnId == FMemAllocTableColumns::MemoryPageColumnId)
-		{
-			Writer.Append(LexToString(Alloc->GetAddress() / PageSize));
-		}
-		else if (ColumnId == FMemAllocTableColumns::SizeColumnId)
-		{
-			Writer.Append(LexToString(Alloc->GetSize()));
-		}
-		else if (ColumnId == FMemAllocTableColumns::TagColumnId)
-		{
-			const TCHAR* LLMTagPath = AllocationsProvider->GetTagFullPath(Alloc->GetTag());
-			Writer.AppendString(LLMTagPath ? LLMTagPath : TEXT(""));
-		}
-		else if (ColumnId == FMemAllocTableColumns::AssetColumnId)
-		{
-			Writer.AppendString(*AssetName);
-		}
-		else if (ColumnId == FMemAllocTableColumns::ClassNameColumnId)
-		{
-			Writer.AppendString(*ClassName);
-		}
-		else if (ColumnId == FMemAllocTableColumns::PackageColumnId)
-		{
-			Writer.AppendString(*PackageName);
-		}
-		else if (ColumnId == FMemAllocTableColumns::AllocThreadColumnId)
-		{
-			Writer.Append(LexToString(Alloc->GetAllocThreadId()));
-		}
-		else if (ColumnId == FMemAllocTableColumns::FreeThreadColumnId)
-		{
-			Writer.Append(LexToString(Alloc->GetFreeThreadId()));
-		}
-		else if (ColumnId == FMemAllocTableColumns::AllocFunctionColumnId)
-		{
-			Writer.AppendString(*GetTopFunction(AllocCallstack));
+			const uint32 EventDistance = EndEventIndex - StartEventIndex;
+			Writer.AppendValue(EventDistance);
 		}
-		else if (ColumnId == FMemAllocTableColumns::FreeFunctionColumnId)
-		{
-			Writer.AppendString(*GetTopFunction(FreeCallstack));
-		}
-		else if (ColumnId == FMemAllocTableColumns::AllocSourceFileColumnId)
+	}
+	else if (ColumnId == FMemAllocTableColumns::StartTimeColumnId)
+	{
+		Writer.AppendValueAsTimestamp(Alloc->GetStartTime());
+	}
+	else if (ColumnId == FMemAllocTableColumns::EndTimeColumnId)
+	{
+		Writer.AppendValueAsTimestamp(Alloc->GetEndTime());
+	}
+	else if (ColumnId == FMemAllocTableColumns::DurationColumnId)
+	{
+		Writer.AppendValueAsDuration(Alloc->GetEndTime() - Alloc->GetStartTime());
+	}
+	else if (ColumnId == FMemAllocTableColumns::AddressColumnId)
+	{
+		Writer.AppendValueAsHex64(Alloc->GetAddress());
+	}
+	else if (ColumnId == FMemAllocTableColumns::MemoryPageColumnId)
+	{
+		Writer.AppendValueAsHex64(Alloc->GetAddress() & ~(Params.PageSize - 1)); // see FMemAllocTable::GetAddressPage
+	}
+	else if (ColumnId == FMemAllocTableColumns::SizeColumnId)
+	{
+		Writer.AppendValue(Alloc->GetSize());
+	}
+	else if (ColumnId == FMemAllocTableColumns::TagColumnId)
+	{
+		const TCHAR* LLMTagPath = Params.AllocationsProvider->GetTagFullPath(Alloc->GetTag());
+		if (LLMTagPath)
 		{
-			Writer.AppendString(*GetTopSourceFile(AllocCallstack));
+			Writer.AppendValue(LLMTagPath);
 		}
-		else if (ColumnId == FMemAllocTableColumns::FreeSourceFileColumnId)
-		{
-			Writer.AppendString(*GetTopSourceFile(FreeCallstack));
+	}
+	else if (ColumnId == FMemAllocTableColumns::AssetColumnId)
+	{
+		Writer.AppendValue(Params.AssetName);
+	}
+	else if (ColumnId == FMemAllocTableColumns::ClassNameColumnId)
+	{
+		Writer.AppendValue(Params.ClassName);
+	}
+	else if (ColumnId == FMemAllocTableColumns::PackageColumnId)
+	{
+		Writer.AppendValue(Params.PackageName);
+	}
+	else if (ColumnId == FMemAllocTableColumns::AllocThreadColumnId)
+	{
+		Writer.AppendValue(Alloc->GetAllocThreadId());
+	}
+	else if (ColumnId == FMemAllocTableColumns::FreeThreadColumnId)
+	{
+		Writer.AppendValue(Alloc->GetFreeThreadId());
+	}
+	else if (ColumnId == FMemAllocTableColumns::AllocFunctionColumnId)
+	{
+		Writer.AppendValue(GetTopFunction(Params.AllocCallstack));
+	}
+	else if (ColumnId == FMemAllocTableColumns::FreeFunctionColumnId)
+	{
+		Writer.AppendValue(GetTopFunction(Params.FreeCallstack));
+	}
+	else if (ColumnId == FMemAllocTableColumns::AllocSourceFileColumnId)
+	{
+		Writer.AppendValue(GetTopSourceFile(Params.AllocCallstack));
+	}
+	else if (ColumnId == FMemAllocTableColumns::FreeSourceFileColumnId)
+	{
+		Writer.AppendValue(GetTopSourceFile(Params.FreeCallstack));
+	}
+	else if (ColumnId == FMemoryExporter::AllocCallstackColumnId)
+	{
+		if (Writer.IsCSV())
+		{
+#if 0
+			// Slower, but it would also escape the " char, if present.
+			Writer.AppendValue(FormatCallstack(Params.AllocCallstack));
+#else
+			// Faster, but we assume the callstack string does not contains the " char
+			// (so we can Append it directly without escaping " char).
+			Writer.AppendDoubleQuotationMark();
+			AppendCallstack(Writer, Params.AllocCallstack);
+			Writer.AppendDoubleQuotationMark();
+#endif
 		}
-		else if (ColumnId == ExportMemoryLeakSnapshot_AllocCallstack)
+		else
 		{
-			Writer.AppendString(*FormatCallstackForCSV(AllocCallstack));
+			AppendCallstack(Writer, Params.AllocCallstack);
 		}
-		else if (ColumnId == ExportMemoryLeakSnapshot_FreeCallstack)
-		{
-			Writer.AppendString(*FormatCallstackForCSV(FreeCallstack));
+	}
+	else if (ColumnId == FMemoryExporter::FreeCallstackColumnId)
+	{
+		if (Writer.IsCSV())
+		{
+#if 0
+			// Slower, but it would also escape the " char, if present.
+			Writer.AppendValue(FormatCallstack(Params.FreeCallstack));
+#else
+			// Faster, but we assume the callstack string does not contains the " char
+			// (so we can Append it directly without escaping " char).
+			Writer.AppendDoubleQuotationMark();
+			AppendCallstack(Writer, Params.FreeCallstack);
+			Writer.AppendDoubleQuotationMark();
+#endif
 		}
 		else
 		{
-			Writer.AppendString(TEXT(""));
+			AppendCallstack(Writer, Params.FreeCallstack);
 		}
 	}
+	else
+	{
+		//Writer.Append(TEXT(""));
+	}
+}
 
-	////////////////////////////////////////////////////////////////////////////////////////////////////
+////////////////////////////////////////////////////////////////////////////////////////////////////
 
+void FMemoryExporter::ExtractMetadata(
+	const TraceServices::IAllocationsProvider::FAllocation* Alloc,
+	FMemoryExporter::FWriteColumnValueParams& InOutParams)
+{
+	InOutParams.AssetName = nullptr;
+	InOutParams.ClassName = nullptr;
+	InOutParams.PackageName = nullptr;
 
-	void FMemoryExporter::ExtractMetadata(const TraceServices::IAllocationsProvider::FAllocation* Allocation,
-		uint16 AssetMetadataType,
-		const TraceServices::IMetadataProvider* MetadataProvider,
-		const TraceServices::FMetadataSchema* MetadataSchema,
-		const TraceServices::IDefinitionProvider* DefinitionProvider,
-		FString& OutAssetName, FString& OutClassName, FString& OutPackageName)
+	const uint32 MetadataId = Alloc->GetMetadataId();
+	if (MetadataId == TraceServices::IMetadataProvider::InvalidMetadataId)
 	{
-		const uint32 MetadataId = Allocation->GetMetadataId();
-		if (MetadataId != TraceServices::IMetadataProvider::InvalidMetadataId &&
-			MetadataProvider && DefinitionProvider && MetadataSchema)
+		return;
+	}
+
+	check(InOutParams.MetadataProvider);
+	check(InOutParams.DefinitionProvider);
+	check(InOutParams.MetadataSchema);
+
+	TraceServices::FProviderReadScopeLock MetadataProviderReadLock(*InOutParams.MetadataProvider);
+	TraceServices::FProviderReadScopeLock DefinitionProviderReadLock(*InOutParams.DefinitionProvider);
+
+	InOutParams.MetadataProvider->EnumerateMetadata(Alloc->GetAllocThreadId(), MetadataId,
+		[&InOutParams]
+		(uint32 StackDepth, uint16 Type, const void* Data, uint32 Size) -> bool
 		{
-			TraceServices::FProviderReadScopeLock MetadataProviderReadLock(*MetadataProvider);
-			TraceServices::FProviderReadScopeLock DefinitionProviderReadLock(*DefinitionProvider);
+			if (Type == InOutParams.AssetMetadataType)
+			{
+				const auto Reader = InOutParams.MetadataSchema->Reader();
 
-			MetadataProvider->EnumerateMetadata(Allocation->GetAllocThreadId(), MetadataId,
-				[AssetMetadataType, MetadataSchema, &OutAssetName, &OutClassName, &OutPackageName, DefinitionProvider]
-			(uint32 StackDepth, uint16 Type, const void* Data, uint32 Size) -> bool
+				// Asset Name
+				const auto AssetNameRef = Reader.GetValueAs<UE::Trace::FEventRef32>((uint8*)Data, 0);
+				const auto AssetNameDef = InOutParams.DefinitionProvider->Get<TraceServices::FStringDefinition>(*AssetNameRef);
+				if (AssetNameDef)
 				{
-					if (Type == AssetMetadataType)
-					{
-						const auto Reader = MetadataSchema->Reader();
-
-						// Asset Name
-						const auto AssetNameRef = Reader.GetValueAs<UE::Trace::FEventRef32>((uint8*)Data, 0);
-						const auto AssetNameDef = DefinitionProvider->Get<TraceServices::FStringDefinition>(*AssetNameRef);
-						if (AssetNameDef)
-						{
-							OutAssetName = FString(AssetNameDef->Display);
-						}
-
-						// Class Name
-						const auto ClassNameRef = Reader.GetValueAs<UE::Trace::FEventRef32>((uint8*)Data, 1);
-						const auto ClassNameDef = DefinitionProvider->Get<TraceServices::FStringDefinition>(*ClassNameRef);
-						if (ClassNameDef)
-						{
-							OutClassName = FString(ClassNameDef->Display);
-						}
-
-						// Package Name
-						const auto PackageNameRef = Reader.GetValueAs<UE::Trace::FEventRef32>((uint8*)Data, 2);
-						const auto PackageNameDef = DefinitionProvider->Get<TraceServices::FStringDefinition>(*PackageNameRef);
-						if (PackageNameDef)
-						{
-							OutPackageName = FString(PackageNameDef->Display);
-						}
-
-						return false;
-					}
-					return true;
-				});
-		}
+					InOutParams.AssetName = AssetNameDef->Display;
+				}
+
+				// Class Name
+				const auto ClassNameRef = Reader.GetValueAs<UE::Trace::FEventRef32>((uint8*)Data, 1);
+				const auto ClassNameDef = InOutParams.DefinitionProvider->Get<TraceServices::FStringDefinition>(*ClassNameRef);
+				if (ClassNameDef)
+				{
+					InOutParams.ClassName = ClassNameDef->Display;
+				}
+
+				// Package Name
+				const auto PackageNameRef = Reader.GetValueAs<UE::Trace::FEventRef32>((uint8*)Data, 2);
+				const auto PackageNameDef = InOutParams.DefinitionProvider->Get<TraceServices::FStringDefinition>(*PackageNameRef);
+				if (PackageNameDef)
+				{
+					InOutParams.PackageName = PackageNameDef->Display;
+				}
+
+				return false;
+			}
+			return true;
+		});
+}
+
+////////////////////////////////////////////////////////////////////////////////////////////////////
+
+bool FMemoryExporter::ExportAllocations(
+	const FString& FilePath,
+	const TArray<const TraceServices::IAllocationsProvider::FAllocation*>& Allocations,
+	const TArray<FName>& ExportColumns)
+{
+	const TraceServices::IAllocationsProvider* AllocationsProvider = TraceServices::ReadAllocationsProvider(Session);
+	if (!AllocationsProvider)
+	{
+		UE_LOG(LogMemoryExporter, Error, TEXT("AllocationsProvider not available"));
+		return false;
 	}
 
-	////////////////////////////////////////////////////////////////////////////////////////////////////
+	const TraceServices::ICallstacksProvider* CallstacksProvider = TraceServices::ReadCallstacksProvider(Session);
+	const TraceServices::IMetadataProvider* MetadataProvider = TraceServices::ReadMetadataProvider(Session);
+	const TraceServices::IDefinitionProvider* DefinitionProvider = TraceServices::ReadDefinitionProvider(Session);
+
+	TraceServices::FAnalysisSessionReadScope SessionReadScope(Session);
+	TraceServices::FProviderReadScopeLock AllocProviderLock(*AllocationsProvider);
 
-	bool FMemoryExporter::WriteSnapshotCSV(const FString& FilePath, 
-		const TArray<const TraceServices::IAllocationsProvider::FAllocation*>& Allocations,
-		const TArray<FName>& ExportColumns)
+	// User requires data for Asset/Class/Package extraction. Setting up the Metadata schema here
+	uint16 AssetMetadataType = TraceServices::IMetadataProvider::InvalidMetadataType;
+	const TraceServices::FMetadataSchema* MetadataSchema = nullptr;
+	if (MetadataProvider)
 	{
-		// Deal with sessions and providers first
-		TraceServices::FAnalysisSessionReadScope SessionReadScope(Session);
-		const TraceServices::IAllocationsProvider* AllocationsProvider = TraceServices::ReadAllocationsProvider(Session);
-		const TraceServices::ICallstacksProvider* CallstacksProvider = TraceServices::ReadCallstacksProvider(Session);
-		const TraceServices::IMetadataProvider* MetadataProvider = TraceServices::ReadMetadataProvider(Session);
-		const TraceServices::IDefinitionProvider* DefinitionProvider = TraceServices::ReadDefinitionProvider(Session);
-		if (!AllocationsProvider)
+		TraceServices::FProviderReadScopeLock MetadataProviderLock(*MetadataProvider);
+		AssetMetadataType = MetadataProvider->GetRegisteredMetadataType(TEXT("Asset"));
+		if (AssetMetadataType != TraceServices::IMetadataProvider::InvalidMetadataType)
 		{
-			UE_LOG(LogMemoryLeakSnapshot, Error, TEXT("AllocationsProvider not available"));
-			return false;
+			MetadataSchema = MetadataProvider->GetRegisteredMetadataSchema(AssetMetadataType);
 		}
-		TraceServices::FProviderReadScopeLock AllocProviderLock(*AllocationsProvider);
+		else
+		{
+			// If Asset metadata type not found, disable metadata provider
+			MetadataProvider = nullptr;
+		}
+	}
+	const bool bCanAccessMetadata = MetadataProvider && MetadataSchema && DefinitionProvider;
 
+	FStopwatch Stopwatch;
+	Stopwatch.Start();
 
-		// User requires data for Asset/Class/Package extraction. Setting up the Metadata schema here
-		uint16 AssetMetadataType = TraceServices::IMetadataProvider::InvalidMetadataType;
-		const TraceServices::FMetadataSchema* MetadataSchema = nullptr;
-		if (MetadataProvider)
+	TUniquePtr<IFileHandle> ExportFileHandle(FPlatformFileManager::Get().GetPlatformFile().OpenWrite(*FilePath));
+	if (!ExportFileHandle)
+	{
+		UE_LOG(LogMemoryExporter, Error, TEXT("Failed to open file for writing: %s"), *FilePath);
+		return false;
+	}
+	bool bIsCSV = FilePath.EndsWith(TEXT(".csv"));
+	FUtf8FileWriter Writer(ExportFileHandle.Get(), bIsCSV);
+
+	constexpr uint64 PageSize = 65536;
+
+	FWriteColumnValueParams Params
+	{
+		PageSize,
+		AllocationsProvider,
+		MetadataProvider,
+		MetadataSchema,
+		DefinitionProvider,
+		AssetMetadataType,
+	};
+
+	// Write header
+	for (int32 ColumnIndex = 0; ColumnIndex < ExportColumns.Num(); ++ColumnIndex)
+	{
+		if (ColumnIndex > 0)
 		{
-			TraceServices::FProviderReadScopeLock MetadataProviderLock(*MetadataProvider);
-			AssetMetadataType = MetadataProvider->GetRegisteredMetadataType(TEXT("Asset"));
-			if (AssetMetadataType != TraceServices::IMetadataProvider::InvalidMetadataType)
-			{
-				MetadataSchema = MetadataProvider->GetRegisteredMetadataSchema(AssetMetadataType);
-			}
-			else
-			{
-				// If Asset metadata type not found, disable metadata provider
-				MetadataProvider = nullptr;
-			}
+			Writer.AppendSeparator();
 		}
+		Writer.Append(ExportColumns[ColumnIndex].ToString());
+	}
+	Writer.AppendLineEnd();
 
+	const uint32 NumAllocations = Allocations.Num();
+	uint32 NumWrittenAllocations = 0;
+	uint32 StepNumWrittenAllocations = 0;
+	int32 StepProgress = 0;
+	FStopwatch StepProgressStopwatch;
+	StepProgressStopwatch.Start();
 
-		// Building the CSV data
-		TUniquePtr<IFileHandle> FileHandle(FPlatformFileManager::Get().GetPlatformFile().OpenWrite(*FilePath));
-		if (!FileHandle)
+	// Write values
+	for (const TraceServices::IAllocationsProvider::FAllocation* Alloc : Allocations)
+	{
+		// Resolve callstacks
+		if (CallstacksProvider)
 		{
-			UE_LOG(LogMemoryLeakSnapshot, Error, TEXT("Failed to open file for writing: %s"), *FilePath);
-			return false;
+			Params.AllocCallstack = CallstacksProvider->GetCallstack(Alloc->GetAllocCallstackId());
+			Params.FreeCallstack  = CallstacksProvider->GetCallstack(Alloc->GetFreeCallstackId());
 		}
-		bool bIsCSV = FilePath.EndsWith(TEXT(".csv"));
-		FUtf16Writer Writer(FileHandle.Get(), bIsCSV);
 
-		// Write header - Generic support for N ExportColumns
-		for (int32 i = 0; i < ExportColumns.Num(); i++)
+		// Resolve metadata
+		if (bCanAccessMetadata)
 		{
-			Writer.Append(ExportColumns[i].ToString());
-			if (i < ExportColumns.Num() - 1)
+			ExtractMetadata(Alloc, Params);
+		}
+
+		for (int32 ColumnIndex = 0; ColumnIndex < ExportColumns.Num(); ++ColumnIndex)
+		{
+			if (ColumnIndex > 0)
 			{
 				Writer.AppendSeparator();
 			}
+			WriteColumnValue(Writer, ExportColumns[ColumnIndex], Alloc, Params);
 		}
 		Writer.AppendLineEnd();
 
-		// Write values
-		const uint32 MaxEventIndex = TNumericLimits<uint32>::Max();
-		const uint64 PageSize = 65536;
-		for (const TraceServices::IAllocationsProvider::FAllocation* Alloc : Allocations)
+		++NumWrittenAllocations;
+		double Progress = ((double)NumWrittenAllocations / (double)NumAllocations) * 100.0f;
+		int32 RoundedProgress = FMath::FloorToInt32(Progress);
+		if (RoundedProgress != StepProgress)
 		{
-			// Resolve callstacks
-			const TraceServices::FCallstack* AllocCallstack = nullptr;
-			const TraceServices::FCallstack* FreeCallstack = nullptr;
-			if (CallstacksProvider)
-			{
-				AllocCallstack = CallstacksProvider->GetCallstack(Alloc->GetAllocCallstackId());
-				FreeCallstack = CallstacksProvider->GetCallstack(Alloc->GetFreeCallstackId());
-			}
+			StepProgressStopwatch.Stop();
+			const double Duration = StepProgressStopwatch.GetAccumulatedTime();
+			const double AllocsPerSecond = (double)(NumWrittenAllocations - StepNumWrittenAllocations) / Duration;
+			UE_LOG(LogMemoryExporter, Log, TEXT("Export progress: %i%% (%u / %u) -> %.0f allocs/s"), RoundedProgress, NumWrittenAllocations, NumAllocations, AllocsPerSecond);
+			StepNumWrittenAllocations = NumWrittenAllocations;
+			StepProgress = RoundedProgress;
+			StepProgressStopwatch.Restart();
+		}
+	}
+	check(NumWrittenAllocations == NumAllocations);
 
-			// Resolve Metadata
-			FString AssetName;
-			FString ClassName;
-			FString PackageName;
-			ExtractMetadata(Alloc, AssetMetadataType, MetadataProvider, MetadataSchema, DefinitionProvider,
-				AssetName, ClassName, PackageName);
+	Writer.Flush();
+	ExportFileHandle->Flush();
+	ExportFileHandle.Reset();
 
-			// Writing Values - Generic
-			for (int32 i = 0; i < ExportColumns.Num(); ++i)
-			{
-				WriteColumnValue(Writer, ExportColumns[i], Alloc, AllocationsProvider, AllocCallstack, FreeCallstack,
-					AssetName, ClassName, PackageName, PageSize);
-				if (i < ExportColumns.Num() - 1)
-				{
-					Writer.AppendSeparator();
-				}
-			}
-			Writer.AppendLineEnd();
-		}
+	Stopwatch.Stop();
+	const double TotalTime = Stopwatch.GetAccumulatedTime();
+	const double TotalAllocsPerSecond = (double)NumWrittenAllocations / TotalTime;
+	UE_LOG(LogMemoryExporter, Log, TEXT("Exported %d allocations to file in %.3fs (\"%s\") -> %.0f allocs/s"), NumWrittenAllocations, TotalTime, *FilePath, TotalAllocsPerSecond);
 
-		Writer.Flush();
-		UE_LOG(LogMemoryLeakSnapshot, Log, TEXT("Written %d rows to %s"), Allocations.Num(), *FilePath);
-		return true;
-	}
+	return true;
+}
 
-	////////////////////////////////////////////////////////////////////////////////////////////////////
+////////////////////////////////////////////////////////////////////////////////////////////////////
 
-	FString FMemoryExporter::GetTopFunction(const TraceServices::FCallstack* Callstack)
+FUtf8String FMemoryExporter::GetTopFunction(const TraceServices::FCallstack* Callstack)
+{
+	if (!Callstack || Callstack->IsEmpty())
 	{
-		if (!Callstack || Callstack->IsEmpty())
-		{
-			return TEXT("No Callstack Recorded");
-		}
+		return UTF8TEXT("No Callstack Recorded");
+	}
 
+	if (Callstack->Num() > 0)
+	{
 		const TraceServices::FStackFrame* Frame = Callstack->Frame(0);
 		if (Frame && Frame->Symbol)
 		{
-			TStringBuilder<512> Builder;
+			TUtf8StringBuilder<512> Builder;
+
+			// "Module!Function" format
 			if (Frame->Symbol->Module)
 			{
 				Builder.Append(Frame->Symbol->Module);
-				Builder.Append(TEXT("!"));
+				Builder.Append(UTF8TEXTVIEW("!"));
 			}
 			if (Frame->Symbol->Name)
 			{
@@ -542,179 +632,178 @@ namespace UE::Insights::MemoryProfiler
 			}
 			else
 			{
-				Builder.Append(TEXT("0x"));
-				Builder.Append(LexToString(Frame->Addr));
+				// Address (as hex number) if no symbol name
+				Builder.Appendf(UTF8TEXT("0x%llX"), Frame->Addr);
 			}
+
 			return Builder.ToString();
 		}
-
-		return TEXT("No Callstack Recorded");
 	}
 
-	////////////////////////////////////////////////////////////////////////////////////////////////////
+	return UTF8TEXT("Empty Callstack");
+}
 
-	FString FMemoryExporter::GetTopSourceFile(const TraceServices::FCallstack* Callstack)
+////////////////////////////////////////////////////////////////////////////////////////////////////
+
+FUtf8String FMemoryExporter::GetTopSourceFile(const TraceServices::FCallstack* Callstack)
+{
+	if (!Callstack || Callstack->IsEmpty())
 	{
-		if (!Callstack || Callstack->IsEmpty())
-		{
-			return TEXT("No Callstack Recorded");
-		}
+		return UTF8TEXT("No Callstack Recorded");
+	}
 
+	if (Callstack->Num() > 0)
+	{
 		const TraceServices::FStackFrame* Frame = Callstack->Frame(0);
 		if (Frame && Frame->Symbol && Frame->Symbol->File)
 		{
-			return FString(Frame->Symbol->File);
+			return FUtf8String(Frame->Symbol->File);
 		}
-
-		return TEXT("No Callstack Recorded");
 	}
 
-	////////////////////////////////////////////////////////////////////////////////////////////////////
+	return UTF8TEXT("Empty Callstack");
+}
 
-	FString FMemoryExporter::FormatCallstackForCSV(const TraceServices::FCallstack* Callstack)
+////////////////////////////////////////////////////////////////////////////////////////////////////
+
+FUtf8String FMemoryExporter::FormatCallstack(const TraceServices::FCallstack* Callstack)
+{
+	if (!Callstack || Callstack->IsEmpty())
 	{
-		if (!Callstack || Callstack->IsEmpty())
-		{
-			return TEXT("No Callstack Recorded");
-		}
+		return UTF8TEXT("No Callstack Recorded");
+	}
 
-		TStringBuilder<4096> Builder;
-		const uint32 NumFrames = Callstack->Num();
-		check(NumFrames < 256);
+	TUtf8StringBuilder<4096> Builder;
+	const uint32 NumFrames = Callstack->Num();
+	check(NumFrames < 256);
 
-		for (uint8 i = 0; i < NumFrames; ++i)
+	for (uint8 i = 0; i < NumFrames; ++i)
+	{
+		const TraceServices::FStackFrame* Frame = Callstack->Frame(i);
+		if (Frame && Frame->Symbol)
 		{
-			const TraceServices::FStackFrame* Frame = Callstack->Frame(i);
-			if (Frame && Frame->Symbol)
+			if (i > 0)
 			{
-				if (i > 0)
-				{
-					Builder.Append(TEXT("/"));
-				}
+				Builder.Append(UTF8TEXTVIEW("/"));
+			}
 
-				// Module!Function format
-				if (Frame->Symbol->Module)
-				{
-					Builder.Append(Frame->Symbol->Module);
-					Builder.Append(TEXT("!"));
-				}
-				if (Frame->Symbol->Name)
-				{
-					Builder.Append(Frame->Symbol->Name);
-				}
-				else
-				{
-					// Address if no symbol name
-					Builder.Append(TEXT("0x"));
-					Builder.Append(LexToString(Frame->Addr));
-				}
+			// "Module!Function" format
+			if (Frame->Symbol->Module)
+			{
+				Builder.Append(Frame->Symbol->Module);
+				Builder.Append(UTF8TEXTVIEW("!"));
+			}
+			if (Frame->Symbol->Name)
+			{
+				Builder.Append(Frame->Symbol->Name);
+			}
+			else
+			{
+				// Address (as hex number) if no symbol name
+				Builder.Appendf(UTF8TEXT("0x%llX"), Frame->Addr);
 			}
 		}
+	}
 
-		if (Builder.Len() == 0)
-		{
-			return TEXT("No Callstack Recorded");
-		}
+	if (Builder.Len() == 0)
+	{
+		return UTF8TEXT("Empty Callstack");
+	}
 
-		return Builder.ToString();
+	return Builder.ToString();
+}
+
+////////////////////////////////////////////////////////////////////////////////////////////////////
+
+void FMemoryExporter::AppendCallstack(FUtf8FileWriter& Writer, const TraceServices::FCallstack* Callstack)
+{
+	if (!Callstack || Callstack->IsEmpty())
+	{
+		Writer.Append(UTF8TEXTVIEW("No Callstack Recorded"));
+		return;
 	}
 
-	////////////////////////////////////////////////////////////////////////////////////////////////////
+	const uint32 NumFrames = Callstack->Num();
+	check(NumFrames < 256);
+	bool bIsEmpty = true;
 
-	bool FMemoryExporter::FindBookmarkByName(const FString& BookmarkName, double& OutTime) const
+	for (uint8 i = 0; i < NumFrames; ++i)
 	{
-		TraceServices::FAnalysisSessionReadScope SessionReadScope(Session);
-		const TraceServices::IBookmarkProvider& BookmarkProvider = TraceServices::ReadBookmarkProvider(Session);
-		const uint64 TotalBookmarks = BookmarkProvider.GetBookmarkCount();
-		if (TotalBookmarks == 0)
+		const TraceServices::FStackFrame* Frame = Callstack->Frame(i);
+		if (Frame && Frame->Symbol)
 		{
-			UE_LOG(LogMemoryLeakSnapshot, Warning, TEXT("No bookmarks found in trace"));
-			return false;
-		}
-		
-		bool BookmarkExists = false;
-		BookmarkProvider.EnumerateBookmarks(
-			0.0,
-			TNumericLimits<double>::Max(),
-			[&BookmarkName, &OutTime, &BookmarkExists](const TraceServices::FBookmark& Bookmark)
+			if (i > 0)
 			{
-				if (Bookmark.Text && FCString::Strcmp(Bookmark.Text, *BookmarkName) == 0)
-				{
-					OutTime = Bookmark.Time;
-					BookmarkExists = true;
-				}
-			});
+				Writer.Append(UTF8TEXTVIEW("/"));
+			}
 
-		if (BookmarkExists)
-		{
-			UE_LOG(LogMemoryLeakSnapshot, Log, TEXT("Found bookmark '%s' at time %.3f"), *BookmarkName, OutTime);
-		}
-		else
-		{
-			UE_LOG(LogMemoryLeakSnapshot, Warning, TEXT("Bookmark '%s' not found"), *BookmarkName);
+			// "Module!Function" format
+			if (Frame->Symbol->Module)
+			{
+				Writer.Append(Frame->Symbol->Module);
+				Writer.Append(UTF8TEXTVIEW("!"));
+			}
+			if (Frame->Symbol->Name)
+			{
+				Writer.Append(Frame->Symbol->Name);
+			}
+			else
+			{
+				// Address (as hex number) if no symbol name
+				Writer.Appendf(UTF8TEXT("0x%llX"), Frame->Addr);
+			}
+			bIsEmpty = false;
 		}
-
-		return BookmarkExists;
 	}
 
-	////////////////////////////////////////////////////////////////////////////////////////////////////
-
-	void FMemoryExporter::FUtf16Writer::WriteBOM()
+	if (bIsEmpty)
 	{
-		UTF16CHAR BOM = UNICODE_BOM;
-		FileHandle->Write((uint8*)&BOM, sizeof(UTF16CHAR));
+		Writer.Append(UTF8TEXTVIEW("Empty Callstack"));
 	}
+}
 
-	////////////////////////////////////////////////////////////////////////////////////////////////////
+////////////////////////////////////////////////////////////////////////////////////////////////////
 
-	void FMemoryExporter::FUtf16Writer::WriteStringBuilder(int32 CacheLen)
-	{
-		if (StringBuilder.Len() > CacheLen)
-		{
-			FTCHARToUTF16 UTF16String(StringBuilder.ToString(), StringBuilder.Len());
-			FileHandle->Write((const uint8*)UTF16String.Get(), UTF16String.Length() * sizeof(UTF16CHAR));
-			StringBuilder.Reset();
-		}
-	}
+bool FMemoryExporter::FindBookmarkByName(const FString& BookmarkName, double& OutTime) const
+{
+	const TraceServices::IBookmarkProvider& BookmarkProvider = TraceServices::ReadBookmarkProvider(Session);
 
-	////////////////////////////////////////////////////////////////////////////////////////////////////
+	TraceServices::FAnalysisSessionReadScope SessionReadScope(Session);
 
-	void FMemoryExporter::FUtf16Writer::AppendString(const TCHAR* Str)
+	const uint64 TotalBookmarks = BookmarkProvider.GetBookmarkCount();
+	if (TotalBookmarks == 0)
 	{
-		if (Str == nullptr || Str[0] == TCHAR('\0'))
-		{
-			// nothing to append
-		}
-		else if (FCString::Strchr(Str, TCHAR(Separator)) != nullptr)
+		UE_LOG(LogMemoryExporter, Warning, TEXT("No bookmarks found in trace"));
+		return false;
+	}
+
+	bool BookmarkExists = false;
+	BookmarkProvider.EnumerateBookmarks(
+		0.0,
+		std::numeric_limits<double>::max(),
+		[&BookmarkName, &OutTime, &BookmarkExists]
+		(const TraceServices::FBookmark& Bookmark)
 		{
-			if (FCString::Strchr(Str, TEXT('\"')) != nullptr)
+			if (Bookmark.Text && FCString::Strcmp(Bookmark.Text, *BookmarkName) == 0)
 			{
-				FString String = Str;
-				String.ReplaceInline(TEXT("\""), TEXT("\"\""));
-				String = FString::Printf(TEXT("\"%s\""), *String);
-				FTCHARToUTF16 UTF16String(*String);
-				StringBuilder.Append(UTF16String.Get(), UTF16String.Length());
+				OutTime = Bookmark.Time;
+				BookmarkExists = true;
 			}
-			else
-			{
-				FString String = FString::Printf(TEXT("\"%s\""), Str);
-				FTCHARToUTF16 UTF16String(*String);
-				StringBuilder.Append(UTF16String.Get(), UTF16String.Length());
-			}
-		}
-		else
-		{
-			FTCHARToUTF16 UTF16String(Str);
-			StringBuilder.Append(UTF16String.Get(), UTF16String.Length());
-		}
-	}
-
-	////////////////////////////////////////////////////////////////////////////////////////////////////
+		});
 
-	void FMemoryExporter::FUtf16Writer::Append(const FString& InString)
+	if (BookmarkExists)
 	{
-		FTCHARToUTF16 UTF16String(*InString, InString.Len());
-		StringBuilder.Append(UTF16String.Get(), UTF16String.Length());
+		UE_LOG(LogMemoryExporter, Log, TEXT("Found bookmark '%s' at time %.3f"), *BookmarkName, OutTime);
 	}
+	else
+	{
+		UE_LOG(LogMemoryExporter, Warning, TEXT("Bookmark '%s' not found"), *BookmarkName);
+	}
+
+	return BookmarkExists;
+}
+
+////////////////////////////////////////////////////////////////////////////////////////////////////
+////////////////////////////////////////////////////////////////////////////////////////////////////
 
-}
\ No newline at end of file
+} // namespace UE::Insights::MemoryProfiler
diff --git a/Engine/Source/Developer/TraceInsights/Private/Insights/MemoryProfiler/ViewModels/MemoryExporter.h b/Engine/Source/Developer/TraceInsights/Private/Insights/MemoryProfiler/ViewModels/MemoryExporter.h
--- a/Engine/Source/Developer/TraceInsights/Private/Insights/MemoryProfiler/ViewModels/MemoryExporter.h
+++ b/Engine/Source/Developer/TraceInsights/Private/Insights/MemoryProfiler/ViewModels/MemoryExporter.h
@@ -3,152 +3,113 @@
 #pragma once
 
 #include "CoreTypes.h"
+
 #include "Containers/Array.h"
-#include "Containers/Map.h"
-#include "Containers/Set.h"
-#include "HAL/Platform.h"
-#include "Templates/Function.h"
-#include "Templates/UniquePtr.h"
 #include "UObject/NameTypes.h"
 
-// TraceServices
-#include "Insights/MemoryProfiler/ViewModels/MemAllocTable.h"
+//TraceServices
 #include "TraceServices/Model/AllocationsProvider.h"
-#include "TraceServices/Model/AnalysisSession.h"
-#include "TraceServices/Model/Callstack.h"
-
-// Standard Library
-#include <limits>
+#include "TraceServices/Model/MetadataProvider.h"
 
-// Forward Declarations
 class IFileHandle;
+
 namespace TraceServices
 {
-	class FMetadataSchema;
 	class IDefinitionProvider;
-	class IMetadataProvider;
+	struct FCallstack;
 }
 
+namespace UE::Insights
+{
+	class FUtf8FileWriter;
+}
 
 namespace UE::Insights::MemoryProfiler
 {
+
 struct FExportMemoryLeakSnapshotParams;
+
 class FMemoryExporter
 {
-private:
-	static constexpr int32 StringBuilderBufferSize = 64 * 1024;
 public:
 	/**
-	* The list of parameters to run memory leak snapshot
+	* The list of parameters to export memory allocs
 	*/
-	struct FExportMemoryLeakSnapshotParams
+	struct FExportMemoryAllocsParams
 	{
 		FString OutputPathName;
 		FString Columns;
-		double TimeA, TimeB, TimeC;
+		TraceServices::IAllocationsProvider::EQueryRule Rule;
+		double TimeA, TimeB, TimeC, TimeD;
 		int32 MaxResults;
 	};
-	// Just like TimeExporter.h, we created a writer class
-	// Why not UTF8?: SMemAllocTableTreeView.cpp uses UTF16 and decided to comply with the standard. 
-	// We made it public though, so other future cmd line exporters can use this helper
-	class FUtf16Writer
-	{
-	public:
-		FUtf16Writer(IFileHandle* InFileHandle, bool bIsCSV)
-			: FileHandle(InFileHandle)
-			, Separator(bIsCSV ? UTF16CHAR(',') : UTF16CHAR('\t'))
-		{
-			check(FileHandle);
-			WriteBOM();
-		}
-
-		~FUtf16Writer() = default;
-
-		TStringBuilder<StringBuilderBufferSize>& GetStringBuilder() { return StringBuilder; }
-		UTF16CHAR GetSeparator() const { return Separator; }
-		UTF16CHAR GetLineEnd() const { return UTF16CHAR('\n'); }
-
-		void AppendSeparator()
-		{
-			StringBuilder.AppendChar(Separator);
-		}
-
-		void AppendLineEnd()
-		{
-			StringBuilder.AppendChar(UTF16CHAR('\n'));
-			WriteStringBuilder(StringBuilderBufferSize - 1024);
-		}
-
-		void AppendString(const TCHAR* InString);
-		void Append(const FString& InString);
-
-		void Flush()
-		{
-			WriteStringBuilder(0);
-		}
-
-	private:
-		void WriteBOM();
-		void WriteStringBuilder(int32 CacheLen);
-
-	private:
-		IFileHandle* FileHandle;
-		TStringBuilder<StringBuilderBufferSize> StringBuilder;
-		UTF16CHAR Separator;
 
-	};
+public:
+	FMemoryExporter(const TraceServices::IAnalysisSession& InSession);
+	virtual ~FMemoryExporter();
+
+	//////////////////////////////////////////////////////////////////////
+	// Exporters
+
+	int32 ExportMemoryAllocsAsText(FExportMemoryAllocsParams& Params);
 
-private:
 	//////////////////////////////////////////////////////////////////////
 	// Utilities
-	void ExportMemorySnapshot_InitColumns();
+
+	bool FindBookmarkByName(const FString& BookmarkName, double& OutTime) const;
+
+private:
+	void InitDefaultColumns();
+	void InitAvailableColumns();
+	void MakeColumnList(const FString& InColumnsString, TArray<FName>& OutColumnList);
+
 	bool QueryAllocations(
-		const FExportMemoryLeakSnapshotParams& Params,
+		const FExportMemoryAllocsParams& Params,
 		TArray<const TraceServices::IAllocationsProvider::FAllocation*>& OutAllocations);
-	bool WriteSnapshotCSV(const FString& FilePath,
+
+	bool ExportAllocations(
+		const FString& FilePath,
 		const TArray<const TraceServices::IAllocationsProvider::FAllocation*>& Allocations,
 		const TArray<FName>& ExportColumns);
-	void MakeExportMemoryLeakSnapshotColumnList(const FString& InColumnsString, TArray<FName>& OutColumnList);
-	void InitializeAvailableColumns();
-	FString GetTopFunction(const TraceServices::FCallstack* Callstack);
-	FString GetTopSourceFile(const TraceServices::FCallstack* Callstack);
-	FString FormatCallstackForCSV(const TraceServices::FCallstack* Callstack);
-	void WriteColumnValue(FUtf16Writer& Writer, const FName& ColumnId,
-		const TraceServices::IAllocationsProvider::FAllocation* Alloc,
-		const TraceServices::IAllocationsProvider* AllocationsProvider,
-		const TraceServices::FCallstack* AllocCallstack,
-		const TraceServices::FCallstack* FreeCallstack,
-		const FString& AssetName, const FString& ClassName, const FString& PackageName, uint64 PageSize);
-	void ExtractMetadata(const TraceServices::IAllocationsProvider::FAllocation* Allocation, 
-		uint16 AssetMetadataType,
-		const TraceServices::IMetadataProvider* MetadataProvider,
-		const TraceServices::FMetadataSchema* MetadataSchema,
-		const TraceServices::IDefinitionProvider* DefinitionProvider,
-		FString& OutAssetName, FString& OutClassName, FString& OutPackageName);
 
+	FUtf8String GetTopFunction(const TraceServices::FCallstack* Callstack);
+	FUtf8String GetTopSourceFile(const TraceServices::FCallstack* Callstack);
+	FUtf8String FormatCallstack(const TraceServices::FCallstack* Callstack);
+	void AppendCallstack(FUtf8FileWriter& InWriter, const TraceServices::FCallstack* Callstack);
 
-public:
-	FMemoryExporter(const TraceServices::IAnalysisSession& InSession);
-	virtual ~FMemoryExporter();
+	struct FWriteColumnValueParams
+	{
+		uint64 PageSize = 0;
+		const TraceServices::IAllocationsProvider* AllocationsProvider = nullptr;
+		const TraceServices::IMetadataProvider* MetadataProvider = nullptr;
+		const TraceServices::FMetadataSchema* MetadataSchema = nullptr;
+		const TraceServices::IDefinitionProvider* DefinitionProvider = nullptr;
+		uint16 AssetMetadataType = TraceServices::IMetadataProvider::InvalidMetadataType;
+		const TraceServices::FCallstack* AllocCallstack = nullptr;
+		const TraceServices::FCallstack* FreeCallstack = nullptr;
+		const TCHAR* AssetName = nullptr;
+		const TCHAR* ClassName = nullptr;
+		const TCHAR* PackageName = nullptr;
+	};
 
-	//////////////////////////////////////////////////////////////////////
-	// Exporters
-	int32 ExportMemoryLeakSnapshotAsText(FExportMemoryLeakSnapshotParams& Params);
+	void ExtractMetadata(
+		const TraceServices::IAllocationsProvider::FAllocation* Alloc,
+		FWriteColumnValueParams& InOutParams);
 
-	//////////////////////////////////////////////////////////////////////
-	// Utilities
-	bool FindBookmarkByName(const FString& BookmarkName, double& OutTime) const;
+	void WriteColumnValue(
+		FUtf8FileWriter& InWriter,
+		const FName& InColumnId,
+		const TraceServices::IAllocationsProvider::FAllocation* Alloc,
+		const FWriteColumnValueParams& InParams);
 
 private:
 	const TraceServices::IAnalysisSession& Session;
 
 	TArray<FName> AvailableColumns;
-	TArray<FName> SnapshotDefaultColumns;
-	
-	static const FName ExportMemoryLeakSnapshot_AllocCallstack;
-	static const FName ExportMemoryLeakSnapshot_FreeCallstack;
+	TArray<FName> DefaultColumns;
 
-	
+	static const FName AllocCallstackColumnId;
+	static const FName FreeCallstackColumnId;
 };
-} // namespace UE::Insights::MemoryProfiler
 
+} // namespace UE::Insights::MemoryProfiler
diff --git a/Engine/Source/Developer/TraceInsights/Private/Insights/MemoryProfiler/Widgets/SMemAllocTableTreeView.cpp b/Engine/Source/Developer/TraceInsights/Private/Insights/MemoryProfiler/Widgets/SMemAllocTableTreeView.cpp
--- a/Engine/Source/Developer/TraceInsights/Private/Insights/MemoryProfiler/Widgets/SMemAllocTableTreeView.cpp
+++ b/Engine/Source/Developer/TraceInsights/Private/Insights/MemoryProfiler/Widgets/SMemAllocTableTreeView.cpp
@@ -6,8 +6,8 @@
 #include "DesktopPlatformModule.h"
 #include "Framework/MultiBox/MultiBoxBuilder.h"
 #include "GenericPlatform/GenericPlatformMisc.h"
-#include "HAL/PlatformFileManager.h"
 #include "HAL/PlatformApplicationMisc.h"
+#include "HAL/PlatformFileManager.h"
 #include "ISourceCodeAccessModule.h"
 #include "ISourceCodeAccessor.h"
 #include "Logging/MessageLog.h"
@@ -37,6 +37,7 @@
 #include "InsightsCore/Table/ViewModels/TableColumn.h"
 
 // TraceInsights
+#include "Insights/Common/Utf8FileWriter.h"
 #include "Insights/InsightsStyle.h"
 #include "Insights/MemoryProfiler/Common/SymbolSearchPathsHelper.h"
 #include "Insights/MemoryProfiler/MemoryProfilerManager.h"
@@ -1641,58 +1642,34 @@ void SMemAllocTableTreeView::ExportMemorySnapshot() const
 		return;
 	}
 
-	const FString& Path = SaveFilenames[0];
-	const TCHAR Separator = Path.EndsWith(TEXT(".csv")) ? TEXT(',') : TEXT('\t');
-
-	const auto ExportFileHandle = std::unique_ptr<IFileHandle>(FPlatformFileManager::Get().GetPlatformFile().OpenWrite(*Path));
+	const FString& FilePath = SaveFilenames[0];
 
-	if (ExportFileHandle == nullptr)
+	// Building the CSV/TSV data
+	TUniquePtr<IFileHandle> FileHandle(FPlatformFileManager::Get().GetPlatformFile().OpenWrite(*FilePath));
+	if (!FileHandle)
 	{
 		FMessageLog ReportMessageLog((LogListingName != NAME_None) ? LogListingName : TEXT("Other"));
 		ReportMessageLog.Error(LOCTEXT("FailedToOpenFile", "Export failed. Failed to open file for write."));
 		ReportMessageLog.Notify();
 		return;
 	}
+	bool bIsCSV = FilePath.EndsWith(TEXT(".csv"));
+	FUtf8FileWriter Writer(FileHandle.Get(), bIsCSV);
 
-	UTF16CHAR BOM = UNICODE_BOM;
-	ExportFileHandle->Write((uint8*)&BOM, sizeof(UTF16CHAR));
-
-	static constexpr TCHAR LineEnd = TEXT('\n');
-	static constexpr TCHAR QuotationMarkBegin = TEXT('\"');
-	static constexpr TCHAR QuotationMarkEnd = TEXT('\"');
-
-	bool IsFirstRow = true;
-	auto WriteColumnHeader = [Separator](FStringBuilderBase& OutData, const FTableColumn& Column)
+	auto WriteColumnHeader =
+		[](FUtf8FileWriter& Writer, const FTableColumn& Column)
 		{
-			FString Value = Column.GetShortName().ToString().ReplaceCharWithEscapedChar();
-			int32 CharIndex;
-			if (Value.FindChar(Separator, CharIndex))
-			{
-				OutData += QuotationMarkBegin;
-				OutData += Value;
-				OutData += QuotationMarkEnd;
-			}
-			else
-			{
-				OutData += Value;
-			}
+			FString Value = Column.GetShortName().ToString();
+			Writer.AppendValue(Value);
 		};
-	auto WriteColumn = [Separator](FStringBuilderBase& OutData, const FTableColumn& Column, const FTableTreeNode& Node)
+	auto WriteColumn =
+		[](FUtf8FileWriter& Writer, const FTableColumn& Column, const FTableTreeNode& Node)
 		{
-			FString Value = Column.GetValue(Node)->AsString().ReplaceCharWithEscapedChar();
-			int32 CharIndex;
-			if (Value.FindChar(Separator, CharIndex))
-			{
-				OutData << QuotationMarkBegin;
-				OutData << Value;
-				OutData << QuotationMarkEnd;
-			}
-			else
-			{
-				OutData << Value;
-			}
+			FString Value = Column.GetValue(Node)->AsString();
+			Writer.AppendValue(Value);
 		};
-	auto WriteCallstackColumn = [Separator](FStringBuilderBase& OutData, const FTableTreeNode& Node, bool bIsAllocCallstack)
+	auto WriteCallstackColumn =
+		[](FUtf8FileWriter& Writer, const FTableTreeNode& Node, bool bIsAllocCallstack)
 		{
 			if (!Node.Is<FMemAllocNode>())
 			{
@@ -1706,41 +1683,37 @@ void SMemAllocTableTreeView::ExportMemorySnapshot() const
 
 				if (!Callstack)
 				{
-					OutData << QuotationMarkBegin;
-					OutData << GetCallstackNotAvailableString().ToString();
-					OutData << QuotationMarkEnd;
+					Writer.AppendValue(GetCallstackNotAvailableString().ToString());
 					return;
 				}
 
 				if (Callstack->Num() == 0)
 				{
-					OutData << QuotationMarkBegin;
 					if (Callstack->GetEmptyId() == 0)
 					{
-						OutData << GetNoCallstackString().ToString();
+						Writer.AppendValue(GetNoCallstackString().ToString());
 					}
 					else
 					{
-						OutData << GetEmptyCallstackString().ToString();
+						Writer.AppendValue(GetEmptyCallstackString().ToString());
 					}
-					OutData << QuotationMarkEnd;
 					return;
 				}
 
-				OutData << QuotationMarkBegin;
+				TStringBuilder<1024> StringBuilder;
 				const uint32 NumCallstackFrames = Callstack->Num();
 				check(NumCallstackFrames <= 256); // see Callstack->Frame(uint8)
 				for (uint32 Index = 0; Index < NumCallstackFrames; ++Index)
 				{
 					if (Index != 0)
 					{
-						OutData << TEXT("/");
+						StringBuilder.Append(TEXT("/"));
 					}
 					const TraceServices::FStackFrame* Frame = Callstack->Frame(static_cast<uint8>(Index));
 					check(Frame != nullptr);
-					FormatStackFrame(*Frame, OutData, EStackFrameFormatFlags::ModuleAndSymbol);
+					FormatStackFrame(*Frame, StringBuilder, EStackFrameFormatFlags::ModuleAndSymbol);
 				}
-				OutData << QuotationMarkEnd;
+				Writer.AppendValue(StringBuilder.ToView());
 			}
 		};
 	const FTableColumn& StartEventIndexColumn = *GetTable()->FindColumn(FMemAllocTableColumns::StartEventIndexColumnId);
@@ -1764,70 +1737,65 @@ void SMemAllocTableTreeView::ExportMemorySnapshot() const
 	const FTableColumn& FreeSourceFileColumn = *GetTable()->FindColumn(FMemAllocTableColumns::FreeSourceFileColumnId);
 
 	// 2. Iterate over TreeNodes
-	TStringBuilder<2048> Buffer;
+	bool IsFirstRow = true;
 	for (const TSharedPtr<FTableTreeNode>& Node : *FilteredNodesPtr)
 	{
 		// Export only leaves
 		if (Node->IsGroup()) continue;
-		// String buffer optimization
-		Buffer.Reset();
 
 		if (IsFirstRow)
 		{
-			WriteColumnHeader(Buffer, StartEventIndexColumn); Buffer += Separator;
-			WriteColumnHeader(Buffer, EndEventIndexColumn); Buffer += Separator;
-			WriteColumnHeader(Buffer, EventDistanceColumn); Buffer += Separator;
-			WriteColumnHeader(Buffer, StartTimeColumn); Buffer += Separator;
-			WriteColumnHeader(Buffer, EndTimeColumn); Buffer += Separator;
-			WriteColumnHeader(Buffer, DurationColumn); Buffer += Separator;
-			WriteColumnHeader(Buffer, AddressColumn); Buffer += Separator;
-			WriteColumnHeader(Buffer, MemoryPageColumn); Buffer += Separator;
-			WriteColumnHeader(Buffer, SizeColumn); Buffer += Separator;
-			WriteColumnHeader(Buffer, TagColumn); Buffer += Separator;
-			WriteColumnHeader(Buffer, AssetColumn); Buffer += Separator;
-			WriteColumnHeader(Buffer, ClassNameColumn); Buffer += Separator;
-			WriteColumnHeader(Buffer, PackageColumn); Buffer += Separator;
-			WriteColumnHeader(Buffer, AllocThreadColumn); Buffer += Separator;
-			WriteColumnHeader(Buffer, FreeThreadColumn); Buffer += Separator;
-			WriteColumnHeader(Buffer, AllocFunctionColumn); Buffer += Separator;
-			WriteColumnHeader(Buffer, FreeFunctionColumn); Buffer += Separator;
-			WriteColumnHeader(Buffer, AllocSourceFileColumn); Buffer += Separator;
-			WriteColumnHeader(Buffer, FreeSourceFileColumn); Buffer += Separator;
-			Buffer += TEXT("Alloc Callstack"); Buffer += Separator;
-			Buffer += TEXT("Free Callstack"); Buffer += LineEnd;
+			WriteColumnHeader(Writer, StartEventIndexColumn); Writer.AppendSeparator();
+			WriteColumnHeader(Writer, EndEventIndexColumn); Writer.AppendSeparator();
+			WriteColumnHeader(Writer, EventDistanceColumn); Writer.AppendSeparator();
+			WriteColumnHeader(Writer, StartTimeColumn); Writer.AppendSeparator();
+			WriteColumnHeader(Writer, EndTimeColumn); Writer.AppendSeparator();
+			WriteColumnHeader(Writer, DurationColumn); Writer.AppendSeparator();
+			WriteColumnHeader(Writer, AddressColumn); Writer.AppendSeparator();
+			WriteColumnHeader(Writer, MemoryPageColumn); Writer.AppendSeparator();
+			WriteColumnHeader(Writer, SizeColumn); Writer.AppendSeparator();
+			WriteColumnHeader(Writer, TagColumn); Writer.AppendSeparator();
+			WriteColumnHeader(Writer, AssetColumn); Writer.AppendSeparator();
+			WriteColumnHeader(Writer, ClassNameColumn); Writer.AppendSeparator();
+			WriteColumnHeader(Writer, PackageColumn); Writer.AppendSeparator();
+			WriteColumnHeader(Writer, AllocThreadColumn); Writer.AppendSeparator();
+			WriteColumnHeader(Writer, FreeThreadColumn); Writer.AppendSeparator();
+			WriteColumnHeader(Writer, AllocFunctionColumn); Writer.AppendSeparator();
+			WriteColumnHeader(Writer, FreeFunctionColumn); Writer.AppendSeparator();
+			WriteColumnHeader(Writer, AllocSourceFileColumn); Writer.AppendSeparator();
+			WriteColumnHeader(Writer, FreeSourceFileColumn); Writer.AppendSeparator();
+			Writer.Append(UTF8TEXTVIEW("Alloc Callstack")); Writer.AppendSeparator();
+			Writer.Append(UTF8TEXTVIEW("Free Callstack")); Writer.AppendLineEnd();
 
 			IsFirstRow = false;
 		}
 
 		// 3. Export these column values as is:
-		WriteColumn(Buffer, StartEventIndexColumn, *Node); Buffer += Separator;
-		WriteColumn(Buffer, EndEventIndexColumn, *Node); Buffer += Separator;
-		WriteColumn(Buffer, EventDistanceColumn, *Node); Buffer += Separator;
-		WriteColumn(Buffer, StartTimeColumn, *Node); Buffer += Separator;
-		WriteColumn(Buffer, EndTimeColumn, *Node); Buffer += Separator;
-		WriteColumn(Buffer, DurationColumn, *Node); Buffer += Separator;
-		WriteColumn(Buffer, AddressColumn, *Node); Buffer += Separator;
-		WriteColumn(Buffer, MemoryPageColumn, *Node); Buffer += Separator;
-		WriteColumn(Buffer, SizeColumn, *Node); Buffer += Separator;
-		WriteColumn(Buffer, TagColumn, *Node); Buffer += Separator;
-		WriteColumn(Buffer, AssetColumn, *Node); Buffer += Separator;
-		WriteColumn(Buffer, ClassNameColumn, *Node); Buffer += Separator;
-		WriteColumn(Buffer, PackageColumn, *Node); Buffer += Separator;
-		WriteColumn(Buffer, AllocThreadColumn, *Node); Buffer += Separator;
-		WriteColumn(Buffer, FreeThreadColumn, *Node); Buffer += Separator;
-		WriteColumn(Buffer, AllocFunctionColumn, *Node); Buffer += Separator;
-		WriteColumn(Buffer, FreeFunctionColumn, *Node); Buffer += Separator;
-		WriteColumn(Buffer, AllocSourceFileColumn, *Node); Buffer += Separator;
-		WriteColumn(Buffer, FreeSourceFileColumn, *Node); Buffer += Separator;
-		WriteCallstackColumn(Buffer, *Node, true); Buffer += Separator;
-		WriteCallstackColumn(Buffer, *Node, false); Buffer += LineEnd;
-
-		// 5. Write rows to file
-		FTCHARToUTF16 UTF16String(*Buffer, Buffer.Len());
-		ExportFileHandle->Write((const uint8*)UTF16String.Get(), UTF16String.Length() * sizeof(UTF16CHAR));
+		WriteColumn(Writer, StartEventIndexColumn, *Node); Writer.AppendSeparator();
+		WriteColumn(Writer, EndEventIndexColumn, *Node); Writer.AppendSeparator();
+		WriteColumn(Writer, EventDistanceColumn, *Node); Writer.AppendSeparator();
+		WriteColumn(Writer, StartTimeColumn, *Node); Writer.AppendSeparator();
+		WriteColumn(Writer, EndTimeColumn, *Node); Writer.AppendSeparator();
+		WriteColumn(Writer, DurationColumn, *Node); Writer.AppendSeparator();
+		WriteColumn(Writer, AddressColumn, *Node); Writer.AppendSeparator();
+		WriteColumn(Writer, MemoryPageColumn, *Node); Writer.AppendSeparator();
+		WriteColumn(Writer, SizeColumn, *Node); Writer.AppendSeparator();
+		WriteColumn(Writer, TagColumn, *Node); Writer.AppendSeparator();
+		WriteColumn(Writer, AssetColumn, *Node); Writer.AppendSeparator();
+		WriteColumn(Writer, ClassNameColumn, *Node); Writer.AppendSeparator();
+		WriteColumn(Writer, PackageColumn, *Node); Writer.AppendSeparator();
+		WriteColumn(Writer, AllocThreadColumn, *Node); Writer.AppendSeparator();
+		WriteColumn(Writer, FreeThreadColumn, *Node); Writer.AppendSeparator();
+		WriteColumn(Writer, AllocFunctionColumn, *Node); Writer.AppendSeparator();
+		WriteColumn(Writer, FreeFunctionColumn, *Node); Writer.AppendSeparator();
+		WriteColumn(Writer, AllocSourceFileColumn, *Node); Writer.AppendSeparator();
+		WriteColumn(Writer, FreeSourceFileColumn, *Node); Writer.AppendSeparator();
+		WriteCallstackColumn(Writer, *Node, true); Writer.AppendSeparator();
+		WriteCallstackColumn(Writer, *Node, false); Writer.AppendLineEnd();
 	}
 
-	ExportFileHandle->Flush();
+	Writer.Flush();
+	FileHandle->Flush();
 }
 
 ////////////////////////////////////////////////////////////////////////////////////////////////////
diff --git a/Engine/Source/Developer/TraceInsights/Private/Insights/TimingProfiler/ViewModels/TimingExporter.cpp b/Engine/Source/Developer/TraceInsights/Private/Insights/TimingProfiler/ViewModels/TimingExporter.cpp
--- a/Engine/Source/Developer/TraceInsights/Private/Insights/TimingProfiler/ViewModels/TimingExporter.cpp
+++ b/Engine/Source/Developer/TraceInsights/Private/Insights/TimingProfiler/ViewModels/TimingExporter.cpp
@@ -18,6 +18,7 @@
 #include "InsightsCore/Common/Stopwatch.h"
 
 // TraceInsights
+#include "Insights/Common/Utf8FileWriter.h"
 #include "Insights/Log.h"
 #include "Insights/TimingProfiler/TimingProfilerManager.h"
 #include "Insights/TimingProfiler/Tracks/GpuTimingTrack.h"
@@ -75,46 +76,6 @@ void FTimingExporter::Error(const TCHAR* InTitle, const TCHAR* InMessage) const
 
 ////////////////////////////////////////////////////////////////////////////////////////////////////
 
-void FTimingExporter::FUtf8Writer::AppendString(const TCHAR* Str)
-{
-	if (Str == nullptr || Str[0] == TCHAR('\0'))
-	{
-		// nothing to append
-	}
-	else if (FCString::Strchr(Str, TCHAR(Separator)) != nullptr)
-	{
-		if (FCString::Strchr(Str, TEXT('\"')) != nullptr)
-		{
-			FString String = Str;
-			String.ReplaceInline(TEXT("\""), TEXT("\"\""));
-			String = FString::Printf(TEXT("\"%s\""), *String);
-			StringBuilder.Append(TCHAR_TO_UTF8(*String));
-		}
-		else
-		{
-			FString String = FString::Printf(TEXT("\"%s\""), Str);
-			StringBuilder.Append(TCHAR_TO_UTF8(*String));
-		}
-	}
-	else
-	{
-		StringBuilder.Append(TCHAR_TO_UTF8(Str));
-	}
-}
-
-////////////////////////////////////////////////////////////////////////////////////////////////////
-
-void FTimingExporter::FUtf8Writer::WriteStringBuilder(int32 CacheLen)
-{
-	if (StringBuilder.Len() > CacheLen)
-	{
-		FileHandle->Write((const uint8*)StringBuilder.ToString(), StringBuilder.Len() * sizeof(UTF8CHAR));
-		StringBuilder.Reset();
-	}
-}
-
-////////////////////////////////////////////////////////////////////////////////////////////////////
-
 int32 FTimingExporter::ExportThreadsAsText(const FString& Filename, FExportThreadsParams& Params) const
 {
 	const TCHAR* ErrorTitle = TEXT("Failed to export threads!");
@@ -135,16 +96,15 @@ int32 FTimingExporter::ExportThreadsAsText(const FString& Filename, FExportThrea
 		return -1;
 	}
 	bool bIsCSV = Filename.EndsWith(TEXT(".csv"));
-	FUtf8Writer Writer(ExportFileHandle.Get(), bIsCSV);
-	FUtf8StringBuilder& StringBuilder = Writer.GetStringBuilder();
+	FUtf8FileWriter Writer(ExportFileHandle.Get(), bIsCSV);
 
 	// Write header.
 	{
-		StringBuilder.Append(UTF8TEXT("Id"));
+		Writer.Append(UTF8TEXTVIEW("Id"));
 		Writer.AppendSeparator();
-		StringBuilder.Append(UTF8TEXT("Name"));
+		Writer.Append(UTF8TEXTVIEW("Name"));
 		Writer.AppendSeparator();
-		StringBuilder.Append(UTF8TEXT("Group"));
+		Writer.Append(UTF8TEXTVIEW("Group"));
 		Writer.AppendLineEnd();
 	}
 
@@ -154,30 +114,30 @@ int32 FTimingExporter::ExportThreadsAsText(const FString& Filename, FExportThrea
 	{
 		// Old GPU Profiler tracks
 		{
-			StringBuilder.Appendf(UTF8TEXT("%u"), FGpuTimingTrack::Gpu1ThreadId);
+			Writer.AppendValue(FGpuTimingTrack::Gpu1ThreadId);
 			Writer.AppendSeparator();
-			StringBuilder.Append(UTF8TEXT("GPU1"));
+			Writer.Append(UTF8TEXTVIEW("GPU1"));
 			Writer.AppendSeparator();
-			StringBuilder.Append(UTF8TEXT("GPU"));
+			Writer.Append(UTF8TEXTVIEW("GPU"));
 			Writer.AppendLineEnd();
 			++ThreadCount;
 
-			StringBuilder.Appendf(UTF8TEXT("%u"), FGpuTimingTrack::Gpu2ThreadId);
+			Writer.AppendValue(FGpuTimingTrack::Gpu2ThreadId);
 			Writer.AppendSeparator();
-			StringBuilder.Append(UTF8TEXT("GPU2"));
+			Writer.Append(UTF8TEXTVIEW("GPU2"));
 			Writer.AppendSeparator();
-			StringBuilder.Append(UTF8TEXT("GPU"));
+			Writer.Append(UTF8TEXTVIEW("GPU"));
 			Writer.AppendLineEnd();
 			++ThreadCount;
 		}
 
 		// Verse track
 		{
-			StringBuilder.Appendf(UTF8TEXT("%u"), FVerseTimingTrack::VerseThreadId);
+			Writer.AppendValue(FVerseTimingTrack::VerseThreadId);
 			Writer.AppendSeparator();
-			StringBuilder.Append(UTF8TEXT("VerseSampling"));
+			Writer.Append(UTF8TEXTVIEW("VerseSampling"));
 			Writer.AppendSeparator();
-			StringBuilder.Append(UTF8TEXT("Verse"));
+			Writer.Append(UTF8TEXTVIEW("Verse"));
 			Writer.AppendLineEnd();
 			++ThreadCount;
 		}
@@ -189,11 +149,11 @@ int32 FTimingExporter::ExportThreadsAsText(const FString& Filename, FExportThrea
 			const TraceServices::IFrameProvider& FramesProvider = TraceServices::ReadFrameProvider(Session);
 			for (uint32 FrameType = 0; FrameType < (uint32)ETraceFrameType::TraceFrameType_Count; ++FrameType)
 			{
-				StringBuilder.Appendf(UTF8TEXT("%u"), GetFrameTypeNonCollidingId(FrameType));
+				Writer.AppendValue(GetFrameTypeNonCollidingId(FrameType));
 				Writer.AppendSeparator();
-				StringBuilder.Append(FFrameTimingTrack::GetFrameTrackName(FrameType));
+				Writer.Append(FFrameTimingTrack::GetFrameTrackName(FrameType));
 				Writer.AppendSeparator();
-				StringBuilder.Append(UTF8TEXT("Frame Tracks"));
+				Writer.Append(UTF8TEXTVIEW("Frame Tracks"));
 				Writer.AppendLineEnd();
 				++ThreadCount;
 			}
@@ -209,11 +169,11 @@ int32 FTimingExporter::ExportThreadsAsText(const FString& Filename, FExportThrea
 				TimingProfilerProvider->EnumerateGpuQueues(
 					[&](const TraceServices::FGpuQueueInfo& QueueInfo)
 					{
-						StringBuilder.Appendf(UTF8TEXT("%u"), GetNonCollidingId(QueueInfo.Id));
+						Writer.AppendValue(GetNonCollidingId(QueueInfo.Id));
 						Writer.AppendSeparator();
-						StringBuilder.Append(QueueInfo.GetDisplayName());
+						Writer.AppendValue(QueueInfo.GetDisplayName());
 						Writer.AppendSeparator();
-						StringBuilder.Append(UTF8TEXT("GPU"));
+						Writer.Append(UTF8TEXTVIEW("GPU"));
 						Writer.AppendLineEnd();
 						++ThreadCount;
 					});
@@ -229,11 +189,11 @@ int32 FTimingExporter::ExportThreadsAsText(const FString& Filename, FExportThrea
 			ThreadProvider.EnumerateThreads(
 				[&](const TraceServices::FThreadInfo& ThreadInfo)
 				{
-					StringBuilder.Appendf(UTF8TEXT("%u"), ThreadInfo.Id);
+					Writer.AppendValue(ThreadInfo.Id);
 					Writer.AppendSeparator();
-					Writer.AppendString(ThreadInfo.Name);
+					Writer.AppendValue(ThreadInfo.Name);
 					Writer.AppendSeparator();
-					Writer.AppendString(ThreadInfo.GroupName);
+					Writer.AppendValue(ThreadInfo.GroupName);
 					Writer.AppendLineEnd();
 					++ThreadCount;
 				});
@@ -242,7 +202,7 @@ int32 FTimingExporter::ExportThreadsAsText(const FString& Filename, FExportThrea
 
 	Writer.Flush();
 	ExportFileHandle->Flush();
-	ExportFileHandle.Reset(nullptr);
+	ExportFileHandle.Reset();
 
 	Stopwatch.Stop();
 	const double TotalTime = Stopwatch.GetAccumulatedTime();
@@ -273,20 +233,19 @@ int32 FTimingExporter::ExportTimersAsText(const FString& Filename, FExportTimers
 		return -1;
 	}
 	bool bIsCSV = Filename.EndsWith(TEXT(".csv"));
-	FUtf8Writer Writer(ExportFileHandle.Get(), bIsCSV);
-	FUtf8StringBuilder& StringBuilder = Writer.GetStringBuilder();
+	FUtf8FileWriter Writer(ExportFileHandle.Get(), bIsCSV);
 
 	// Write header.
 	{
-		StringBuilder.Append(UTF8TEXT("Id"));
+		Writer.Append(UTF8TEXTVIEW("Id"));
 		Writer.AppendSeparator();
-		StringBuilder.Append(UTF8TEXT("Type"));
+		Writer.Append(UTF8TEXTVIEW("Type"));
 		Writer.AppendSeparator();
-		StringBuilder.Append(UTF8TEXT("Name"));
+		Writer.Append(UTF8TEXTVIEW("Name"));
 		Writer.AppendSeparator();
-		StringBuilder.Append(UTF8TEXT("File"));
+		Writer.Append(UTF8TEXTVIEW("File"));
 		Writer.AppendSeparator();
-		StringBuilder.Append(UTF8TEXT("Line"));
+		Writer.Append(UTF8TEXTVIEW("Line"));
 		Writer.AppendLineEnd();
 	}
 
@@ -309,24 +268,24 @@ int32 FTimingExporter::ExportTimersAsText(const FString& Filename, FExportTimers
 		for (uint32 TimerIndex = 0; TimerIndex < TimerCount; ++TimerIndex)
 		{
 			const TraceServices::FTimingProfilerTimer& Timer = *(TimerReader.GetTimer(TimerIndex));
-			StringBuilder.Appendf(UTF8TEXT("%u"), Timer.Id);
+			Writer.AppendValue(Timer.Id);
 			Writer.AppendSeparator();
 			switch (Timer.Type)
 			{
 			case TraceServices::ETimingProfilerTimerType::GpuScope:
-				StringBuilder.Append(UTF8TEXT("GPU"));
+				Writer.Append(UTF8TEXTVIEW("GPU"));
 				break;
 			case TraceServices::ETimingProfilerTimerType::CpuScope:
-				StringBuilder.Append(UTF8TEXT("CPU"));
+				Writer.Append(UTF8TEXTVIEW("CPU"));
 				break;
 			case TraceServices::ETimingProfilerTimerType::CpuSampling:
-				StringBuilder.Append(UTF8TEXT("CpuSampling"));
+				Writer.Append(UTF8TEXTVIEW("CpuSampling"));
 				break;
 			case TraceServices::ETimingProfilerTimerType::VerseSampling:
-				StringBuilder.Append(UTF8TEXT("VerseSampling"));
+				Writer.Append(UTF8TEXTVIEW("VerseSampling"));
 				break;
 			default:
-				StringBuilder.Append(UTF8TEXT("Unknown"));
+				Writer.Append(UTF8TEXTVIEW("Unknown"));
 				break;
 			}
 			Writer.AppendSeparator();
@@ -334,26 +293,26 @@ int32 FTimingExporter::ExportTimersAsText(const FString& Filename, FExportTimers
 			{
 				TStringBuilder<1024> NameStringBuilder;
 				TimingProfilerProvider->MakeTimerDisplayName(Timer, NameStringBuilder);
-				Writer.AppendString(*NameStringBuilder);
+				Writer.AppendValue(NameStringBuilder.ToView());
 			}
 			else
 			{
-				Writer.AppendString(Timer.Name);
+				Writer.AppendValue(Timer.Name);
 			}
 			Writer.AppendSeparator();
 			if (Timer.File)
 			{
-				StringBuilder.Append((const UTF8CHAR*)TCHAR_TO_UTF8(Timer.File));
+				Writer.AppendValue(Timer.File);
 			}
 			Writer.AppendSeparator();
-			StringBuilder.Appendf(UTF8TEXT("%u"), Timer.Line);
+			Writer.AppendValue(Timer.Line);
 			Writer.AppendLineEnd();
 		}
 	}
 
 	Writer.Flush();
 	ExportFileHandle->Flush();
-	ExportFileHandle.Reset(nullptr);
+	ExportFileHandle.Reset();
 
 	Stopwatch.Stop();
 	const double TotalTime = Stopwatch.GetAccumulatedTime();
@@ -442,7 +401,7 @@ void FTimingExporter::ExportTimingEvents_WriteHeader(FExportTimingEventsInternal
 			{
 				Params.Writer.AppendSeparator();
 			}
-			Params.Writer.Append(Column.GetPlainNameString());
+			Params.Writer.AppendValue(Column.GetPlainNameString());
 		}
 	}
 	Params.Writer.AppendLineEnd();
@@ -472,7 +431,6 @@ int32 FTimingExporter::ExportTimingEvents_WriteEvents(FExportTimingEventsInterna
 					{
 						if (!Params.UserParams.TimingEventFilter || Params.UserParams.TimingEventFilter(EventStartTime, EventEndTime, EventDepth, Event, &TimerReader))
 						{
-							FUtf8StringBuilder& StringBuilder = Params.Writer.GetStringBuilder();
 							bool bFirst = true;
 							for (const FName& Column : Params.Columns)
 							{
@@ -489,11 +447,11 @@ int32 FTimingExporter::ExportTimingEvents_WriteEvents(FExportTimingEventsInterna
 
 									if (Column == ExportTimingEvents_ThreadIdColumn)
 									{
-										StringBuilder.Appendf(UTF8TEXT("%u"), Params.ThreadId);
+										Params.Writer.AppendValue(Params.ThreadId);
 									}
 									else if (Column == ExportTimingEvents_ThreadNameColumn)
 									{
-										Params.Writer.AppendString(Params.ThreadName);
+										Params.Writer.AppendValue(Params.ThreadName);
 									}
 									else if (Column == ExportTimingEvents_TimerIdColumn)
 									{
@@ -503,35 +461,35 @@ int32 FTimingExporter::ExportTimingEvents_WriteEvents(FExportTimingEventsInterna
 											TimerIndex = TimerReader.GetOriginalTimerIdFromMetadata(TimerIndex);
 										}
 
-										StringBuilder.Appendf(UTF8TEXT("%u"), TimerIndex);
+										Params.Writer.AppendValue(TimerIndex);
 									}
 									else if (Column == ExportTimingEvents_TimerNameColumn)
 									{
 										const TraceServices::FTimingProfilerTimer* Timer = TimerReader.GetTimer(Event.TimerIndex);
 										if (Timer != nullptr)
 										{
-											Params.Writer.AppendString(Timer->Name);
+											Params.Writer.AppendValue(Timer->Name);
 										}
 										else
 										{
-											StringBuilder.Appendf(UTF8TEXT("[0x%x]"), Event.TimerIndex);
+											Params.Writer.GetStringBuilder().Appendf(UTF8TEXT("[0x%x]"), Event.TimerIndex);
 										}
 									}
 									else if (Column == ExportTimingEvents_StartTimeColumn)
 									{
-										StringBuilder.Appendf(UTF8TEXT("%.9g"), EventStartTime);
+										Params.Writer.AppendValueAsTimestamp(EventStartTime);
 									}
 									else if (Column == ExportTimingEvents_EndTimeColumn)
 									{
-										StringBuilder.Appendf(UTF8TEXT("%.9g"), EventEndTime);
+										Params.Writer.AppendValueAsTimestamp(EventEndTime);
 									}
 									else if (Column == ExportTimingEvents_DurationColumn)
 									{
-										StringBuilder.Appendf(UTF8TEXT("%.9f"), EventEndTime - EventStartTime);
+										Params.Writer.AppendValueAsDuration(EventEndTime - EventStartTime);
 									}
 									else if (Column == ExportTimingEvents_DepthColumn)
 									{
-										StringBuilder.Appendf(UTF8TEXT("%u"), EventDepth);
+										Params.Writer.AppendValue(EventDepth);
 									}
 								}
 							}
@@ -548,7 +506,6 @@ int32 FTimingExporter::ExportTimingEvents_WriteEvents(FExportTimingEventsInterna
 			[&Params, &TimingEventCount]
 			(const TraceServices::FFrame& Frame)
 			{
-				FUtf8StringBuilder& StringBuilder = Params.Writer.GetStringBuilder();
 				bool bFirst = true;
 				for (const FName& Column : Params.Columns)
 				{
@@ -565,35 +522,35 @@ int32 FTimingExporter::ExportTimingEvents_WriteEvents(FExportTimingEventsInterna
 
 						if (Column == ExportTimingEvents_ThreadIdColumn)
 						{
-							StringBuilder.Appendf(UTF8TEXT("%u"), Params.ThreadId);
+							Params.Writer.AppendValue(Params.ThreadId);
 						}
 						else if (Column == ExportTimingEvents_ThreadNameColumn)
 						{
-							Params.Writer.AppendString(Params.ThreadName);
+							Params.Writer.AppendValue(Params.ThreadName);
 						}
 						else if (Column == ExportTimingEvents_TimerIdColumn)
 						{
-							StringBuilder.Appendf(UTF8TEXT("%u"), (uint32)Frame.Index);
+							Params.Writer.AppendValue((uint32)Frame.Index);
 						}
 						else if (Column == ExportTimingEvents_TimerNameColumn)
 						{
-							Params.Writer.AppendString(*FFrameTimingTrack::GetFrameName((uint32)Frame.FrameType, Frame.Index));
+							Params.Writer.AppendValue(*FFrameTimingTrack::GetFrameName((uint32)Frame.FrameType, Frame.Index));
 						}
 						else if (Column == ExportTimingEvents_StartTimeColumn)
 						{
-							StringBuilder.Appendf(UTF8TEXT("%.9g"), Frame.StartTime);
+							Params.Writer.AppendValueAsTimestamp(Frame.StartTime);
 						}
 						else if (Column == ExportTimingEvents_EndTimeColumn)
 						{
-							StringBuilder.Appendf(UTF8TEXT("%.9g"), Frame.EndTime);
+							Params.Writer.AppendValueAsTimestamp(Frame.EndTime);
 						}
 						else if (Column == ExportTimingEvents_DurationColumn)
 						{
-							StringBuilder.Appendf(UTF8TEXT("%.9f"), Frame.EndTime - Frame.StartTime);
+							Params.Writer.AppendValueAsDuration(Frame.EndTime - Frame.StartTime);
 						}
 						else if (Column == ExportTimingEvents_DepthColumn)
 						{
-							StringBuilder.Appendf(UTF8TEXT("%u"), 1u);
+							Params.Writer.Append(UTF8TEXTVIEW("1"));
 						}
 					}
 				}
@@ -757,8 +714,7 @@ int32 FTimingExporter::ExportTimingEventsAsText(const FString& Filename, FExport
 		return -1;
 	}
 	bool bIsCSV = Filename.EndsWith(TEXT(".csv"));
-	FUtf8Writer Writer(ExportFileHandle.Get(), bIsCSV);
-	FUtf8StringBuilder& StringBuilder = Writer.GetStringBuilder();
+	FUtf8FileWriter Writer(ExportFileHandle.Get(), bIsCSV);
 
 	FExportTimingEventsInternalParams InternalParams = { *this, Params, Columns, Writer, 0 };
 
@@ -770,7 +726,7 @@ int32 FTimingExporter::ExportTimingEventsAsText(const FString& Filename, FExport
 
 	Writer.Flush();
 	ExportFileHandle->Flush();
-	ExportFileHandle.Reset(nullptr);
+	ExportFileHandle.Reset();
 
 	Stopwatch.Stop();
 	const double TotalTime = Stopwatch.GetAccumulatedTime();
@@ -1139,28 +1095,33 @@ int32 FTimingExporter::ExportTimerCalleesAsText(const FString& Filename, const F
 	}
 
 	const bool bIsCSV = Filename.EndsWith(TEXT(".csv"));
-	FUtf8Writer Writer(ExportFileHandle.Get(), bIsCSV);
+	FUtf8FileWriter Writer(ExportFileHandle.Get(), bIsCSV);
 
 	// Write header
 	{
-		TArray<FStringView> Columns
+		TArray<FUtf8StringView> Columns
 		{
-			TEXTVIEW("TimerId"),
-			TEXTVIEW("ParentId"),
-			TEXTVIEW("TimerName"),
-			TEXTVIEW("Count"),
-			TEXTVIEW("Inc.Time"),
-			TEXTVIEW("Exc.Time"),
-			TEXTVIEW("NumFrames")
+			UTF8TEXTVIEW("TimerId"),
+			UTF8TEXTVIEW("ParentId"),
+			UTF8TEXTVIEW("TimerName"),
+			UTF8TEXTVIEW("Count"),
+			UTF8TEXTVIEW("Inc.Time"),
+			UTF8TEXTVIEW("Exc.Time"),
+			UTF8TEXTVIEW("NumFrames")
 		};
 
-		for (const FStringView ColumnName : Columns)
+		bool bFirst = true;
+		for (FUtf8StringView ColumnName : Columns)
 		{
-			Writer.AppendString(ColumnName.GetData());
-			if (ColumnName != Columns.Last())
+			if (bFirst)
+			{
+				bFirst = false;
+			}
+			else
 			{
 				Writer.AppendSeparator();
 			}
+			Writer.Append(ColumnName);
 		}
 		Writer.AppendLineEnd();
 	}
@@ -1180,9 +1141,6 @@ int32 FTimingExporter::ExportTimerCalleesAsText(const FString& Filename, const F
 
 	// Write rows
 	{
-		FUtf8StringBuilder& StringBuilder = Writer.GetStringBuilder();
-		TraceServices::FAnalysisSessionReadScope SessionReadScope(Session);
-
 		for (uint32 TimerId : Params.TimerIds)
 		{
 			TArray<const TraceServices::FTimingProfilerButterflyNode*> NodesToVisit{ &Butterfly->GenerateCalleesTree(TimerId) };
@@ -1197,28 +1155,35 @@ int32 FTimingExporter::ExportTimerCalleesAsText(const FString& Filename, const F
 
 				NodesToVisit.Append(CurrentNode->Children);
 
-				StringBuilder.Appendf(UTF8TEXT("%u"), CurrentNode->Timer->Id);
+				Writer.AppendValue(CurrentNode->Timer->Id);
 				Writer.AppendSeparator();
-				StringBuilder.Appendf(UTF8TEXT("%u"), CurrentNode->Parent ? CurrentNode->Parent->Timer->Id : uint32(-1));
+				if (CurrentNode->Parent)
+				{
+					Writer.AppendValue(CurrentNode->Parent->Timer->Id);
+				}
+				else
+				{
+					Writer.Append(UTF8TEXTVIEW("-1"));
+				}
 				Writer.AppendSeparator();
 				if (Params.bShowFormatInTimerNames)
 				{
 					TStringBuilder<1024> NameStringBuilder;
 					TimingProfilerProvider->MakeTimerDisplayName(*CurrentNode->Timer, NameStringBuilder);
-					Writer.AppendString(*NameStringBuilder);
+					Writer.AppendValue(NameStringBuilder.ToView());
 				}
 				else
 				{
-					Writer.AppendString(CurrentNode->Timer->Name);
+					Writer.AppendValue(CurrentNode->Timer->Name);
 				}
 				Writer.AppendSeparator();
-				StringBuilder.Appendf(UTF8TEXT("%llu"), CurrentNode->Count);
+				Writer.AppendValue(CurrentNode->Count);
 				Writer.AppendSeparator();
-				StringBuilder.Appendf(UTF8TEXT("%.9g"), CurrentNode->InclusiveTime);
+				Writer.AppendValueAsDuration(CurrentNode->InclusiveTime);
 				Writer.AppendSeparator();
-				StringBuilder.Appendf(UTF8TEXT("%.9g"), CurrentNode->ExclusiveTime);
+				Writer.AppendValueAsDuration(CurrentNode->ExclusiveTime);
 				Writer.AppendSeparator();
-				StringBuilder.Appendf(UTF8TEXT("%llu"), NumFrames);
+				Writer.AppendValue(NumFrames);
 				Writer.AppendLineEnd();
 			}
 		}
@@ -1505,16 +1470,15 @@ int32 FTimingExporter::ExportCountersAsText(const FString& Filename, FExportCoun
 		return -1;
 	}
 	bool bIsCSV = Filename.EndsWith(TEXT(".csv"));
-	FUtf8Writer Writer(ExportFileHandle.Get(), bIsCSV);
-	FUtf8StringBuilder& StringBuilder = Writer.GetStringBuilder();
+	FUtf8FileWriter Writer(ExportFileHandle.Get(), bIsCSV);
 
 	// Write header.
 	{
-		StringBuilder.Append(UTF8TEXT("Id"));
+		Writer.Append(UTF8TEXTVIEW("Id"));
 		Writer.AppendSeparator();
-		StringBuilder.Append(UTF8TEXT("Type"));
+		Writer.Append(UTF8TEXTVIEW("Type"));
 		Writer.AppendSeparator();
-		StringBuilder.Append(UTF8TEXT("Name"));
+		Writer.Append(UTF8TEXTVIEW("Name"));
 		Writer.AppendLineEnd();
 	}
 
@@ -1530,22 +1494,22 @@ int32 FTimingExporter::ExportCountersAsText(const FString& Filename, FExportCoun
 		CounterProvider.EnumerateCounters(
 			[&](uint32 CounterId, const TraceServices::ICounter& Counter)
 			{
-				StringBuilder.Appendf(UTF8TEXT("%u"), CounterId);
+				Writer.AppendValue(CounterId);
 				Writer.AppendSeparator();
 				if (Counter.IsFloatingPoint())
 				{
-					StringBuilder.Append(UTF8TEXT("Double"));
+					Writer.Append(UTF8TEXTVIEW("Double"));
 				}
 				else
 				{
-					StringBuilder.Append(UTF8TEXT("Int64"));
+					Writer.Append(UTF8TEXTVIEW("Int64"));
 				}
 				if (Counter.IsResetEveryFrame())
 				{
-					StringBuilder.Append(UTF8TEXT("|ResetEveryFrame"));
+					Writer.Append(UTF8TEXTVIEW("|ResetEveryFrame"));
 				}
 				Writer.AppendSeparator();
-				Writer.AppendString(Counter.GetName());
+				Writer.AppendValue(Counter.GetName());
 				Writer.AppendLineEnd();
 				++CounterCount;
 			});
@@ -1553,7 +1517,7 @@ int32 FTimingExporter::ExportCountersAsText(const FString& Filename, FExportCoun
 
 	Writer.Flush();
 	ExportFileHandle->Flush();
-	ExportFileHandle.Reset(nullptr);
+	ExportFileHandle.Reset();
 
 	Stopwatch.Stop();
 	const double TotalTime = Stopwatch.GetAccumulatedTime();
@@ -1659,24 +1623,23 @@ int32 FTimingExporter::ExportCounterAsText(const FString& FilenamePattern, uint3
 		return -1;
 	}
 	bool bIsCSV = Filename.EndsWith(TEXT(".csv"));
-	FUtf8Writer Writer(ExportFileHandle.Get(), bIsCSV);
-	FUtf8StringBuilder& StringBuilder = Writer.GetStringBuilder();
+	FUtf8FileWriter Writer(ExportFileHandle.Get(), bIsCSV);
 
 	// Write header.
 	if (Params.bExportOps)
 	{
-		StringBuilder.Append(UTF8TEXT("Time"));
+		Writer.Append(UTF8TEXTVIEW("Time"));
 		Writer.AppendSeparator();
-		StringBuilder.Append(UTF8TEXT("Op"));
+		Writer.Append(UTF8TEXTVIEW("Op"));
 		Writer.AppendSeparator();
-		StringBuilder.Append(UTF8TEXT("Value"));
+		Writer.Append(UTF8TEXTVIEW("Value"));
 		Writer.AppendLineEnd();
 	}
 	else
 	{
-		StringBuilder.Append(UTF8TEXT("Time"));
+		Writer.Append(UTF8TEXTVIEW("Time"));
 		Writer.AppendSeparator();
-		StringBuilder.Append(UTF8TEXT("Value"));
+		Writer.Append(UTF8TEXTVIEW("Value"));
 		Writer.AppendLineEnd();
 	}
 
@@ -1700,21 +1663,21 @@ int32 FTimingExporter::ExportCounterAsText(const FString& FilenamePattern, uint3
 						Counter.EnumerateFloatOps(Params.IntervalStartTime, Params.IntervalEndTime, false,
 							[&](double Time, TraceServices::ECounterOpType Op, double Value)
 							{
-								StringBuilder.Appendf(UTF8TEXT("%.9f"), Time);
+								Writer.AppendValueAsTimestamp(Time);
 								Writer.AppendSeparator();
 								switch (Op)
 								{
 									case TraceServices::ECounterOpType::Set:
-										StringBuilder.Append(UTF8TEXT("Set"));
+										Writer.Append(UTF8TEXTVIEW("Set"));
 										break;
 									case TraceServices::ECounterOpType::Add:
-										StringBuilder.Append(UTF8TEXT("Add"));
+										Writer.Append(UTF8TEXTVIEW("Add"));
 										break;
 									default:
-										StringBuilder.Appendf(UTF8TEXT("%d"), int32(Op));
+										Writer.AppendValue(int32(Op));
 								}
 								Writer.AppendSeparator();
-								StringBuilder.Appendf(UTF8TEXT("%.9f"), Value);
+								Writer.AppendValue(Value);
 								Writer.AppendLineEnd();
 								++ValueCount;
 							});
@@ -1724,21 +1687,21 @@ int32 FTimingExporter::ExportCounterAsText(const FString& FilenamePattern, uint3
 						Counter.EnumerateOps(Params.IntervalStartTime, Params.IntervalEndTime, false,
 							[&](double Time, TraceServices::ECounterOpType Op, int64 IntValue)
 							{
-								StringBuilder.Appendf(UTF8TEXT("%.9f"), Time);
+								Writer.AppendValueAsTimestamp(Time);
 								Writer.AppendSeparator();
 								switch (Op)
 								{
 									case TraceServices::ECounterOpType::Set:
-										StringBuilder.Append(UTF8TEXT("Set"));
+										Writer.Append(UTF8TEXTVIEW("Set"));
 										break;
 									case TraceServices::ECounterOpType::Add:
-										StringBuilder.Append(UTF8TEXT("Add"));
+										Writer.Append(UTF8TEXTVIEW("Add"));
 										break;
 									default:
-										StringBuilder.Appendf(UTF8TEXT("%d"), int32(Op));
+										Writer.AppendValue(int32(Op));
 								}
 								Writer.AppendSeparator();
-								StringBuilder.Appendf(UTF8TEXT("%lli"), IntValue);
+								Writer.AppendValue(IntValue);
 								Writer.AppendLineEnd();
 								++ValueCount;
 							});
@@ -1751,9 +1714,9 @@ int32 FTimingExporter::ExportCounterAsText(const FString& FilenamePattern, uint3
 						Counter.EnumerateFloatValues(Params.IntervalStartTime, Params.IntervalEndTime, false,
 							[&](double Time, double Value)
 							{
-								StringBuilder.Appendf(UTF8TEXT("%.9f"), Time);
+								Writer.AppendValueAsTimestamp(Time);
 								Writer.AppendSeparator();
-								StringBuilder.Appendf(UTF8TEXT("%.9f"), Value);
+								Writer.AppendValue(Value);
 								Writer.AppendLineEnd();
 								++ValueCount;
 							});
@@ -1763,9 +1726,9 @@ int32 FTimingExporter::ExportCounterAsText(const FString& FilenamePattern, uint3
 						Counter.EnumerateValues(Params.IntervalStartTime, Params.IntervalEndTime, false,
 							[&](double Time, int64 IntValue)
 							{
-								StringBuilder.Appendf(UTF8TEXT("%.9f"), Time);
+								Writer.AppendValueAsTimestamp(Time);
 								Writer.AppendSeparator();
-								StringBuilder.Appendf(UTF8TEXT("%lli"), IntValue);
+								Writer.AppendValue(IntValue);
 								Writer.AppendLineEnd();
 								++ValueCount;
 							});
@@ -1776,7 +1739,7 @@ int32 FTimingExporter::ExportCounterAsText(const FString& FilenamePattern, uint3
 
 	Writer.Flush();
 	ExportFileHandle->Flush();
-	ExportFileHandle.Reset(nullptr);
+	ExportFileHandle.Reset();
 
 	Stopwatch.Stop();
 	const double TotalTime = Stopwatch.GetAccumulatedTime();
diff --git a/Engine/Source/Developer/TraceInsights/Private/Insights/TimingProfiler/ViewModels/TimingExporter.h b/Engine/Source/Developer/TraceInsights/Private/Insights/TimingProfiler/ViewModels/TimingExporter.h
--- a/Engine/Source/Developer/TraceInsights/Private/Insights/TimingProfiler/ViewModels/TimingExporter.h
+++ b/Engine/Source/Developer/TraceInsights/Private/Insights/TimingProfiler/ViewModels/TimingExporter.h
@@ -25,6 +25,11 @@ namespace TraceServices
 	struct FTimingProfilerEvent;
 }
 
+namespace UE::Insights
+{
+	class FUtf8FileWriter;
+}
+
 ////////////////////////////////////////////////////////////////////////////////////////////////////
 
 namespace UE::Insights::TimingProfiler
@@ -200,59 +205,12 @@ public:
 	};
 
 private:
-	static constexpr int32 StringBuilderBufferSize = 32 * 1024;
-	typedef TUtf8StringBuilder<StringBuilderBufferSize> FUtf8StringBuilder;
-
-	class FUtf8Writer
-	{
-	public:
-		FUtf8Writer(IFileHandle* InFileHandle, bool bIsCSV)
-			: FileHandle(InFileHandle)
-			, Separator(bIsCSV ? UTF8CHAR(',') : UTF8CHAR('\t'))
-		{
-			check(FileHandle);
-		}
-
-		~FUtf8Writer() = default;
-
-		FUtf8StringBuilder& GetStringBuilder() { return StringBuilder; }
-		UTF8CHAR GetSeparator() const { return Separator; }
-		UTF8CHAR GetLineEnd() const { return UTF8CHAR('\n'); }
-
-		void AppendSeparator()
-		{
-			StringBuilder.AppendChar(Separator);
-		}
-
-		void AppendLineEnd()
-		{
-			StringBuilder.AppendChar(UTF8CHAR('\n'));
-			WriteStringBuilder(StringBuilderBufferSize - 1024);
-		}
-
-		void AppendString(const TCHAR* InString);
-		void Append(const FString& InString) { StringBuilder.Append(InString); }
-
-		void Flush()
-		{
-			WriteStringBuilder(0);
-		}
-
-	private:
-		void WriteStringBuilder(int32 CacheLen);
-
-	private:
-		IFileHandle* FileHandle;
-		FUtf8StringBuilder StringBuilder;
-		UTF8CHAR Separator;
-	};
-
 	struct FExportTimingEventsInternalParams
 	{
 		const FTimingExporter& Exporter;
 		const FExportTimingEventsParams& UserParams;
 		const TArray<FName>& Columns;
-		FUtf8Writer& Writer;
+		FUtf8FileWriter& Writer;
 		uint32 ThreadId;
 		const TCHAR* ThreadName;
 	};
diff --git a/Engine/Source/Developer/TraceInsights/Private/Insights/TimingProfiler/Widgets/SStatsView.cpp b/Engine/Source/Developer/TraceInsights/Private/Insights/TimingProfiler/Widgets/SStatsView.cpp
--- a/Engine/Source/Developer/TraceInsights/Private/Insights/TimingProfiler/Widgets/SStatsView.cpp
+++ b/Engine/Source/Developer/TraceInsights/Private/Insights/TimingProfiler/Widgets/SStatsView.cpp
@@ -35,6 +35,7 @@
 #include "InsightsCore/Table/Widgets/SAsyncOperationStatus.h"
 
 // TraceInsights
+#include "Insights/Common/Utf8FileWriter.h"
 #include "Insights/InsightsStyle.h"
 #include "Insights/Log.h"
 #include "Insights/Table/ViewModels/TableCommands.h"
@@ -2365,7 +2366,7 @@ void SStatsView::ContextMenu_Export_Execute()
 		return;
 	}
 
-	IFileHandle* ExportFileHandle = OpenExportFile(*Filename);
+	TUniquePtr<IFileHandle> ExportFileHandle{ OpenExportFile(*Filename) };
 	if (!ExportFileHandle)
 	{
 		return;
@@ -2374,19 +2375,8 @@ void SStatsView::ContextMenu_Export_Execute()
 	FStopwatch Stopwatch;
 	Stopwatch.Start();
 
-	UTF16CHAR BOM = UNICODE_BOM;
-	ExportFileHandle->Write((uint8*)&BOM, sizeof(UTF16CHAR));
-
-	TCHAR Separator = TEXT('\t');
-	if (Filename.EndsWith(TEXT(".csv")))
-	{
-		Separator = TEXT(',');
-	}
-	constexpr TCHAR LineEnd = TEXT('\n');
-	constexpr TCHAR QuotationMarkBegin = TEXT('\"');
-	constexpr TCHAR QuotationMarkEnd = TEXT('\"');
-
-	TStringBuilder<1024> StringBuilder;
+	bool bIsCSV = Filename.EndsWith(TEXT(".csv"));
+	FUtf8FileWriter Writer(ExportFileHandle.Get(), bIsCSV);
 
 	TArray<TSharedRef<FTableColumn>> VisibleColumns;
 	Table->GetVisibleColumns(VisibleColumns);
@@ -2402,23 +2392,13 @@ void SStatsView::ContextMenu_Export_Execute()
 			}
 			else
 			{
-				StringBuilder.AppendChar(Separator);
-			}
-			FString Value = ColumnRef->GetShortName().ToString().ReplaceCharWithEscapedChar();
-			int32 CharIndex;
-			if (Value.FindChar(Separator, CharIndex))
-			{
-				StringBuilder.AppendChar(QuotationMarkBegin);
-				StringBuilder.Append(Value);
-				StringBuilder.AppendChar(QuotationMarkEnd);
-			}
-			else
-			{
-				StringBuilder.Append(Value);
+				Writer.AppendSeparator();
 			}
+
+			FString ColumnName = ColumnRef->GetShortName().ToString();
+			Writer.AppendValue(ColumnName);
 		}
-		StringBuilder.AppendChar(LineEnd);
-		ExportFileHandle->Write((const uint8*)StringBuilder.ToString(), StringBuilder.Len() * sizeof(TCHAR));
+		Writer.AppendLineEnd();
 	}
 
 	if (CurrentSorter.IsValid())
@@ -2431,8 +2411,6 @@ void SStatsView::ContextMenu_Export_Execute()
 	{
 		const FBaseTreeNodePtr& Node = SelectedNodes[Index];
 
-		StringBuilder.Reset();
-
 		bool bIsFirstColumn = true;
 		for (const TSharedRef<FTableColumn>& ColumnRef : VisibleColumns)
 		{
@@ -2442,29 +2420,18 @@ void SStatsView::ContextMenu_Export_Execute()
 			}
 			else
 			{
-				StringBuilder.AppendChar(Separator);
+				Writer.AppendSeparator();
 			}
 
-			FString Value = ColumnRef->GetValueAsSerializableString(*Node).ReplaceCharWithEscapedChar();
-			int32 CharIndex;
-			if (Value.FindChar(Separator, CharIndex))
-			{
-				StringBuilder.AppendChar(QuotationMarkBegin);
-				StringBuilder.Append(Value);
-				StringBuilder.AppendChar(QuotationMarkEnd);
-			}
-			else
-			{
-				StringBuilder.Append(Value);
-			}
+			FString Value = ColumnRef->GetValueAsSerializableString(*Node);
+			Writer.AppendValue(Value);
 		}
-		StringBuilder.AppendChar(LineEnd);
-		ExportFileHandle->Write((const uint8*)StringBuilder.ToString(), StringBuilder.Len() * sizeof(TCHAR));
+		Writer.AppendLineEnd();
 	}
 
+	Writer.Flush();
 	ExportFileHandle->Flush();
-	delete ExportFileHandle;
-	ExportFileHandle = nullptr;
+	ExportFileHandle.Reset();
 
 	Stopwatch.Stop();
 	const double TotalTime = Stopwatch.GetAccumulatedTime();
@@ -2656,7 +2623,7 @@ bool SStatsView::OpenSaveTextFileDialog(const FString& InDialogTitle, const FStr
 			InDialogTitle,
 			DefaultPath,
 			InDefaultFile,
-			TEXT("Tab-Separated Values (*.tsv)|*.tsv|Text Files (*.txt)|*.txt|Comma-Separated Values (*.csv)|*.csv|All Files (*.*)|*.*"),
+			TEXT("Comma-Separated Values (*.csv)|*.csv|Tab-Separated Values (*.tsv)|*.tsv|Text Files (*.txt)|*.txt|All Files (*.*)|*.*"),
 			EFileDialogFlags::None,
 			SaveFilenames
 		);
diff --git a/Engine/Source/Developer/TraceInsights/Private/Insights/TimingProfiler/Widgets/STimersView.cpp b/Engine/Source/Developer/TraceInsights/Private/Insights/TimingProfiler/Widgets/STimersView.cpp
--- a/Engine/Source/Developer/TraceInsights/Private/Insights/TimingProfiler/Widgets/STimersView.cpp
+++ b/Engine/Source/Developer/TraceInsights/Private/Insights/TimingProfiler/Widgets/STimersView.cpp
@@ -41,6 +41,7 @@
 #include "InsightsCore/Table/Widgets/SAsyncOperationStatus.h"
 
 // TraceInsights
+#include "Insights/Common/Utf8FileWriter.h"
 #include "Insights/InsightsStyle.h"
 #include "Insights/Log.h"
 #include "Insights/Table/ViewModels/TableCommands.h"
@@ -3355,7 +3356,7 @@ void STimersView::ContextMenu_Export_Execute()
 		return;
 	}
 
-	IFileHandle* ExportFileHandle = OpenExportFile(*Filename);
+	TUniquePtr<IFileHandle> ExportFileHandle{ OpenExportFile(*Filename) };
 	if (!ExportFileHandle)
 	{
 		return;
@@ -3364,19 +3365,8 @@ void STimersView::ContextMenu_Export_Execute()
 	FStopwatch Stopwatch;
 	Stopwatch.Start();
 
-	UTF16CHAR BOM = UNICODE_BOM;
-	ExportFileHandle->Write((uint8*)&BOM, sizeof(UTF16CHAR));
-
-	TCHAR Separator = TEXT('\t');
-	if (Filename.EndsWith(TEXT(".csv")))
-	{
-		Separator = TEXT(',');
-	}
-	constexpr TCHAR LineEnd = TEXT('\n');
-	constexpr TCHAR QuotationMarkBegin = TEXT('\"');
-	constexpr TCHAR QuotationMarkEnd = TEXT('\"');
-
-	TStringBuilder<1024> StringBuilder;
+	bool bIsCSV = Filename.EndsWith(TEXT(".csv"));
+	FUtf8FileWriter Writer(ExportFileHandle.Get(), bIsCSV);
 
 	TArray<TSharedRef<FTableColumn>> VisibleColumns;
 	Table->GetVisibleColumns(VisibleColumns);
@@ -3392,23 +3382,13 @@ void STimersView::ContextMenu_Export_Execute()
 			}
 			else
 			{
-				StringBuilder.AppendChar(Separator);
-			}
-			FString Value = ColumnRef->GetShortName().ToString().ReplaceCharWithEscapedChar();
-			int32 CharIndex;
-			if (Value.FindChar(Separator, CharIndex))
-			{
-				StringBuilder.AppendChar(QuotationMarkBegin);
-				StringBuilder.Append(Value);
-				StringBuilder.AppendChar(QuotationMarkEnd);
-			}
-			else
-			{
-				StringBuilder.Append(Value);
+				Writer.AppendSeparator();
 			}
+
+			FString ColumnName = ColumnRef->GetShortName().ToString();
+			Writer.AppendValue(ColumnName);
 		}
-		StringBuilder.AppendChar(LineEnd);
-		ExportFileHandle->Write((const uint8*)StringBuilder.ToString(), StringBuilder.Len() * sizeof(TCHAR));
+		Writer.AppendLineEnd();
 	}
 
 	if (CurrentSorter.IsValid())
@@ -3421,8 +3401,6 @@ void STimersView::ContextMenu_Export_Execute()
 	{
 		const FBaseTreeNodePtr& Node = SelectedNodes[Index];
 
-		StringBuilder.Reset();
-
 		bool bIsFirstColumn = true;
 		for (const TSharedRef<FTableColumn>& ColumnRef : VisibleColumns)
 		{
@@ -3432,29 +3410,18 @@ void STimersView::ContextMenu_Export_Execute()
 			}
 			else
 			{
-				StringBuilder.AppendChar(Separator);
+				Writer.AppendSeparator();
 			}
 
-			FString Value = ColumnRef->GetValueAsSerializableString(*Node).ReplaceCharWithEscapedChar();
-			int32 CharIndex;
-			if (Value.FindChar(Separator, CharIndex))
-			{
-				StringBuilder.AppendChar(QuotationMarkBegin);
-				StringBuilder.Append(Value);
-				StringBuilder.AppendChar(QuotationMarkEnd);
-			}
-			else
-			{
-				StringBuilder.Append(Value);
-			}
+			FString Value = ColumnRef->GetValueAsSerializableString(*Node);
+			Writer.AppendValue(Value);
 		}
-		StringBuilder.AppendChar(LineEnd);
-		ExportFileHandle->Write((const uint8*)StringBuilder.ToString(), StringBuilder.Len() * sizeof(TCHAR));
+		Writer.AppendLineEnd();
 	}
 
+	Writer.Flush();
 	ExportFileHandle->Flush();
-	delete ExportFileHandle;
-	ExportFileHandle = nullptr;
+	ExportFileHandle.Reset();
 
 	Stopwatch.Stop();
 	const double TotalTime = Stopwatch.GetAccumulatedTime();
diff --git a/Engine/Source/Developer/TraceInsights/Private/Insights/Widgets/SLogView.cpp b/Engine/Source/Developer/TraceInsights/Private/Insights/Widgets/SLogView.cpp
--- a/Engine/Source/Developer/TraceInsights/Private/Insights/Widgets/SLogView.cpp
+++ b/Engine/Source/Developer/TraceInsights/Private/Insights/Widgets/SLogView.cpp
@@ -30,6 +30,7 @@
 #include "InsightsCore/Common/TimeUtils.h"
 
 // TraceInsights
+#include "Insights/Common/Utf8FileWriter.h"
 #include "Insights/InsightsManager.h"
 #include "Insights/InsightsStyle.h"
 #include "Insights/Log.h"
@@ -1695,7 +1696,7 @@ void SLogView::ToggleTimeFormat()
 {
 	bUseRelativeTime = !bUseRelativeTime;
 	TSharedPtr<SHeaderRow> Row = ListView->GetHeaderRow();
-	
+
 	Row->RemoveColumn(Row->GetColumns()[1].ColumnId);
 	Row->InsertColumn(
 		SHeaderRow::Column(bUseRelativeTime ? LogViewColumns::SessionTimeColumnName : LogViewColumns::RealTimeColumnName)
@@ -1889,16 +1890,18 @@ void SLogView::CopyRange() const
 
 		for (const TSharedPtr<FLogMessage>& Message : FilteredMessages)
 		{
-			LogProvider.ReadMessage(Message->GetIndex(), [this, SelectionStartTime, SelectionEndTime, &StringBuilder, &NumMessagesInSelectedRange](const TraceServices::FLogMessageInfo& MessageInfo)
-			{
-				if (MessageInfo.Time >= SelectionStartTime && MessageInfo.Time <= SelectionEndTime)
+			LogProvider.ReadMessage(Message->GetIndex(),
+				[this, SelectionStartTime, SelectionEndTime, &StringBuilder, &NumMessagesInSelectedRange]
+				(const TraceServices::FLogMessageInfo& MessageInfo)
 				{
-					FLogMessageRecord Log(MessageInfo);
-					AppendFormatMessageDelimited(Log, StringBuilder, TEXT('\t'));
-					StringBuilder.AppendChar(TEXT('\n'));
-					++NumMessagesInSelectedRange;
-				}
-			});
+					if (MessageInfo.Time >= SelectionStartTime && MessageInfo.Time <= SelectionEndTime)
+					{
+						FLogMessageRecord Log(MessageInfo);
+						AppendFormatMessageDelimited(Log, StringBuilder, TEXT('\t'));
+						StringBuilder.AppendChar(TEXT('\n'));
+						++NumMessagesInSelectedRange;
+					}
+				});
 		}
 	}
 
@@ -1937,12 +1940,14 @@ void SLogView::CopyAll() const
 
 		for (const TSharedPtr<FLogMessage>& Message : FilteredMessages)
 		{
-			LogProvider.ReadMessage(Message->GetIndex(), [this, &StringBuilder](const TraceServices::FLogMessageInfo& MessageInfo)
-			{
-				FLogMessageRecord Log(MessageInfo);
-				AppendFormatMessageDelimited(Log, StringBuilder, TEXT('\t'));
-				StringBuilder.AppendChar(TEXT('\n'));
-			});
+			LogProvider.ReadMessage(Message->GetIndex(),
+				[this, &StringBuilder]
+				(const TraceServices::FLogMessageInfo& MessageInfo)
+				{
+					FLogMessageRecord Log(MessageInfo);
+					AppendFormatMessageDelimited(Log, StringBuilder, TEXT('\t'));
+					StringBuilder.AppendChar(TEXT('\n'));
+				});
 		}
 	}
 
@@ -2029,15 +2034,10 @@ void SLogView::SaveLogsToFile(bool bSaveLogsInSelectedRangeOnly) const
 		return;
 	}
 
-	FString& Path = SaveFilenames[0];
-	TCHAR Separator = TEXT('\t');
-	if (Path.EndsWith(TEXT(".csv")))
-	{
-		Separator = TEXT(',');
-	}
+	FString& Filename = SaveFilenames[0];
 
-	IFileHandle* ExportFileHandle = FPlatformFileManager::Get().GetPlatformFile().OpenWrite(*Path);
-	if (ExportFileHandle == nullptr)
+	TUniquePtr<IFileHandle> ExportFileHandle{ FPlatformFileManager::Get().GetPlatformFile().OpenWrite(*Filename) };
+	if (!ExportFileHandle)
 	{
 		FMessageLog ReportMessageLog(FInsightsManager::Get()->GetLogListingName());
 		ReportMessageLog.Error(LOCTEXT("FailedToOpenFile", "Save logs failed. Failed to open file for write."));
@@ -2048,16 +2048,16 @@ void SLogView::SaveLogsToFile(bool bSaveLogsInSelectedRangeOnly) const
 	UE::Insights::FStopwatch Stopwatch;
 	Stopwatch.Start();
 
-	UTF16CHAR BOM = UNICODE_BOM;
-	ExportFileHandle->Write((uint8*)&BOM, sizeof(UTF16CHAR));
+	const bool bIsCSV = Filename.EndsWith(TEXT(".csv"));
+	const TCHAR Separator = bIsCSV ? TEXT(',') : TEXT('\t');
+	FUtf8FileWriter Writer(ExportFileHandle.Get(), bIsCSV);
 
 	// Write header.
 	{
-		TStringBuilder<1024> StringBuilder;
-		AppendFormatMessageDelimitedHeader(StringBuilder, Separator);
-		StringBuilder.AppendChar(TEXT('\n'));
-		FTCHARToUTF16 UTF16String(StringBuilder.ToString(), StringBuilder.Len());
-		ExportFileHandle->Write((const uint8*)UTF16String.Get(), UTF16String.Length() * sizeof(UTF16CHAR));
+		TStringBuilder<1024> Line;
+		AppendFormatMessageDelimitedHeader(Line, Separator);
+		Writer.GetStringBuilder().Append(Line.ToView());
+		Writer.AppendLineEnd();
 	}
 
 	int32 NumMessagesInSelectedRange = 0;
@@ -2070,25 +2070,26 @@ void SLogView::SaveLogsToFile(bool bSaveLogsInSelectedRangeOnly) const
 
 		for (const TSharedPtr<FLogMessage>& Message : FilteredMessages)
 		{
-			LogProvider.ReadMessage(Message->GetIndex(), [this, bSaveLogsInSelectedRangeOnly, SelectionStartTime, SelectionEndTime, Separator, ExportFileHandle, &NumMessagesInSelectedRange](const TraceServices::FLogMessageInfo& MessageInfo)
-			{
-				if (!bSaveLogsInSelectedRangeOnly || (MessageInfo.Time >= SelectionStartTime && MessageInfo.Time <= SelectionEndTime))
+			LogProvider.ReadMessage(Message->GetIndex(),
+				[this, bSaveLogsInSelectedRangeOnly, SelectionStartTime, SelectionEndTime, Separator, &Writer, &NumMessagesInSelectedRange]
+				(const TraceServices::FLogMessageInfo& MessageInfo)
 				{
-					TStringBuilder<1024> StringBuilder;
-					FLogMessageRecord Log(MessageInfo);
-					AppendFormatMessageDelimited(Log, StringBuilder, Separator);
-					StringBuilder.AppendChar(TEXT('\n'));
-					FTCHARToUTF16 UTF16String(StringBuilder.ToString(), StringBuilder.Len());
-					ExportFileHandle->Write((const uint8*)UTF16String.Get(), UTF16String.Length() * sizeof(UTF16CHAR));
-					++NumMessagesInSelectedRange;
-				}
-			});
+					if (!bSaveLogsInSelectedRangeOnly || (MessageInfo.Time >= SelectionStartTime && MessageInfo.Time <= SelectionEndTime))
+					{
+						FLogMessageRecord Log(MessageInfo);
+						TStringBuilder<1024> Line;
+						AppendFormatMessageDelimited(Log, Line, Separator);
+						Writer.GetStringBuilder().Append(Line.ToView());
+						Writer.AppendLineEnd();
+						++NumMessagesInSelectedRange;
+					}
+				});
 		}
 	}
 
+	Writer.Flush();
 	ExportFileHandle->Flush();
-	delete ExportFileHandle;
-	ExportFileHandle = nullptr;
+	ExportFileHandle.Reset();
 
 	Stopwatch.Stop();
 	const double TotalTime = Stopwatch.GetAccumulatedTime();
@@ -2112,7 +2113,7 @@ void SLogView::OpenSource() const
 		FLogMessageRecord& Record = Cache.Get(SelectedLogMessage->GetIndex());
 		const FString File = Record.GetFile();
 		const uint32 Line = Record.GetLine();
-		
+
 		ISourceCodeAccessModule& SourceCodeAccessModule = FModuleManager::LoadModuleChecked<ISourceCodeAccessModule>("SourceCodeAccess");
 		if (FPaths::FileExists(File))
 		{
