Add support for searching in dynamic metadata fields in the quick find metadata filter.

#Synced-CL 50479750

diff --git a/Engine/Source/Developer/TraceInsights/Private/Insights/TimingProfiler/Tracks/ThreadTimingTrack.cpp b/Engine/Source/Developer/TraceInsights/Private/Insights/TimingProfiler/Tracks/ThreadTimingTrack.cpp
--- a/Engine/Source/Developer/TraceInsights/Private/Insights/TimingProfiler/Tracks/ThreadTimingTrack.cpp
+++ b/Engine/Source/Developer/TraceInsights/Private/Insights/TimingProfiler/Tracks/ThreadTimingTrack.cpp
@@ -165,6 +165,12 @@ static void AppendMetadataToTooltip(FTooltipDrawState& InOutTooltip, const Trace
 	FCborReader CborReader(&MemoryReader, ECborEndianness::StandardCompliant);
 	FCborContext Context;
 
+	// We have 3 cases:
+	// 1. InMetadataSpec in not null and InMetadata is a list of values (dynamic metadata)
+	// 2. InMetadataSpec is null and InMetadata is a list of values (dynamic metadata without a spec emitted or saved)
+	// 3. InMetadataSpec is null and InMetadata is a map of key value pairs. (static metadata)
+	bool bIsMap = false;
+	bool bHasReadFirstValue = false;
 	if (InMetadataSpec == nullptr)
 	{
 		if (!CborReader.ReadNext(Context))
@@ -173,17 +179,13 @@ static void AppendMetadataToTooltip(FTooltipDrawState& InOutTooltip, const Trace
 			return;
 		}
 
-		if (Context.MajorType() == ECborCode::TextString)
+		if (Context.MajorType() == ECborCode::Map)
 		{
-			InOutTooltip.AddTitle(TEXTVIEW("Metadata..."), FLinearColor(0.3f, 0.3f, 0.3f, 1.0f));
-			InOutTooltip.AddNameValueTextLine(TEXTVIEW(""), Context.AsString());
-			return;
+			bIsMap = true;
 		}
-
-		if (Context.MajorType() != ECborCode::Map)
+		else
 		{
-			InOutTooltip.AddTitle(TEXTVIEW("<invalid metadata>"), FLinearColor(1.0f, 0.5f, 0.5f, 1.0f));
-			return;
+			bHasReadFirstValue = true;
 		}
 	}
 
@@ -205,7 +207,7 @@ static void AppendMetadataToTooltip(FTooltipDrawState& InOutTooltip, const Trace
 				Key = TEXT(""); // undefined field
 			}
 		}
-		else
+		else if (bIsMap)
 		{
 			// Read key
 			if (!CborReader.ReadNext(Context) || !Context.IsString())
@@ -218,9 +220,12 @@ static void AppendMetadataToTooltip(FTooltipDrawState& InOutTooltip, const Trace
 		}
 
 		// Read value
-		if (!CborReader.ReadNext(Context))
+		if (Index > 0 || !bHasReadFirstValue) // Skip reading if we already read the first value.
 		{
-			break;
+			if (!CborReader.ReadNext(Context))
+			{
+				break;
+			}
 		}
 
 		switch (Context.MajorType())
@@ -304,7 +309,7 @@ static void AppendMetadataToTooltip(FTooltipDrawState& InOutTooltip, const Trace
 
 		InOutTooltip.AddNameValueTextLine(Key, TEXTVIEW("???")); // unknown field type
 
-		if (Context.IsFiniteContainer())
+		if (Context.MajorType() == ECborCode::Array)
 		{
 			CborReader.SkipContainer(ECborCode::Array);
 		}
@@ -325,25 +330,30 @@ static void AppendMetadataToString(FString& InOutStr, TArrayView<const uint8>& I
 		return;
 	}
 
-	if (Context.MajorType() == ECborCode::TextString)
+	bool bIsMap = false;
+	if (Context.MajorType() == ECborCode::Map)
 	{
-		InOutStr += TEXTVIEW(" - ");
-		InOutStr += Context.AsString();
-		return;
-	}
-
-	if (Context.MajorType() != ECborCode::Map)
-	{
-		InOutStr += TEXTVIEW(" - <invalid metadata>");
-		return;
+		bIsMap = true;
 	}
 
 	for (uint32 Index = 0; true; ++Index)
 	{
-		// Read key
-		if (!CborReader.ReadNext(Context) || !Context.IsString())
+		if (bIsMap)
 		{
-			break;
+			// Read key
+			if (!CborReader.ReadNext(Context) || !Context.IsString())
+			{
+				break;
+			}
+		}
+
+		if (Index > 0 || bIsMap) // Skip reading if we already read the first value
+		{
+			// Read value
+			if (!CborReader.ReadNext(Context))
+			{
+				break;
+			}
 		}
 
 		if (Index == 0)
@@ -355,16 +365,6 @@ static void AppendMetadataToString(FString& InOutStr, TArrayView<const uint8>& I
 			InOutStr += TEXTVIEW(", ");
 		}
 
-		//FString Key(Context.AsLength(), Context.AsCString());
-		//InOutStr += Key;
-		//InOutStr += TEXTVIEW(":");
-
-		// Read value
-		if (!CborReader.ReadNext(Context))
-		{
-			break;
-		}
-
 		switch (Context.MajorType())
 		{
 
@@ -438,7 +438,7 @@ static void AppendMetadataToString(FString& InOutStr, TArrayView<const uint8>& I
 
 		InOutStr += TEXTVIEW("???"); // unknown field type
 
-		if (Context.IsFiniteContainer())
+		if (Context.MajorType() == ECborCode::Array)
 		{
 			CborReader.SkipContainer(ECborCode::Array);
 		}
@@ -498,7 +498,7 @@ static void AppendMetadataToString(FString& InOutStr, const FString& InFormat, T
 
 	for (uint32 Index = 0; !Format.IsEmpty(); ++Index)
 	{
-		// Read key
+		// Read value
 		if (!CborReader.ReadNext(Context))
 		{
 			break;
@@ -574,7 +574,7 @@ static void AppendMetadataToString(FString& InOutStr, const FString& InFormat, T
 		GetNextFormatSection();
 		AddValueToName(Data, MaxLength, FString(TEXT("%s")), TEXT("???")); // unknown field type
 
-		if (Context.IsFiniteContainer())
+		if (Context.MajorType() == ECborCode::Array)
 		{
 			CborReader.SkipContainer(ECborCode::Array);
 		}
diff --git a/Engine/Source/Developer/TraceInsights/Private/Insights/TimingProfiler/ViewModels/TimerFilters.cpp b/Engine/Source/Developer/TraceInsights/Private/Insights/TimingProfiler/ViewModels/TimerFilters.cpp
--- a/Engine/Source/Developer/TraceInsights/Private/Insights/TimingProfiler/ViewModels/TimerFilters.cpp
+++ b/Engine/Source/Developer/TraceInsights/Private/Insights/TimingProfiler/ViewModels/TimerFilters.cpp
@@ -257,10 +257,16 @@ void FMetadataFilterState::AddCustomUI(TSharedRef<SHorizontalBox> Box)
 
 void FMetadataFilterState::Update()
 {
-	bShowAllMetadataEvents = false;
+	bKeyAllowsAll = false;
+	bTermAllowsAll = false;
+
 	if (Key.Equals(TEXT("*")))
 	{
-		bShowAllMetadataEvents = true;
+		bKeyAllowsAll = true;
+	}
+	if(Term.Equals(TEXT("*")))
+	{
+		bTermAllowsAll = true;
 	}
 
 	ConvertedData = ConvertedDataVariant();
@@ -296,74 +302,135 @@ bool FMetadataFilterState::ApplyFilter(const FBaseFilterContext& Context) const
 	}
 
 	uint32 TimerIndex = static_cast<uint32>(Context.GetDataAsInt64(Filter->GetKey()));
+	if (static_cast<int32>(TimerIndex) >= 0)
+	{
+		return false;
+	}
+	else if (ShouldShowAllMetadataEvents())
+	{
+		return true;
+	}
 
 	TSharedPtr<const TraceServices::IAnalysisSession> Session = FInsightsManager::Get()->GetSession();
-	if (Session.IsValid())
+	if (!Session.IsValid())
 	{
-		const TraceServices::ITimingProfilerProvider* TimingProfilerProvider = TraceServices::ReadTimingProfilerProvider(*Session.Get());
-		if (TimingProfilerProvider)
-		{
-			TraceServices::FAnalysisSessionReadScope SessionReadScope(*Session.Get());
-
-			const TraceServices::ITimingProfilerTimerReader& TimerReader = TimingProfilerProvider->GetTimerReader();
+		return false;
+	}
 
-			TArrayView<const uint8> Metadata = TimerReader.GetMetadata(TimerIndex);
+	const TraceServices::ITimingProfilerProvider* TimingProfilerProvider = TraceServices::ReadTimingProfilerProvider(*Session.Get());
+	if (TimingProfilerProvider == nullptr)
+	{
+		return false;
+	}
 
-			if (Metadata.Num() > 0)
-			{
-				if (!bShowAllMetadataEvents)
-				{
-					return ApplyFilterToMetadata(Metadata);
-				}
+	TraceServices::FAnalysisSessionReadScope SessionReadScope(*Session.Get());
+	const TraceServices::ITimingProfilerTimerReader& TimerReader = TimingProfilerProvider->GetTimerReader();
 
-				return true;
-			}
+	const TraceServices::FTimingProfilerTimer* Timer = TimerReader.GetTimer(TimerIndex);
+	TArrayView<const uint8> Metadata = TimerReader.GetMetadata(TimerIndex);
+	if (Timer->HasValidMetadataSpecId())
+	{
+		const TraceServices::FMetadataSpec* MetadataSpec = nullptr;
+		MetadataSpec = TimingProfilerProvider->GetMetadataSpec(Timer->MetadataSpecId);
+		if (MetadataSpec)
+		{
+			return ApplyFilterToMetadata(MetadataSpec, Metadata);
 		}
 	}
 
+	if (Metadata.Num() > 0)
+	{
+		return ApplyFilterToMetadata(nullptr, Metadata);
+	}
+
 	return false;
 }
 
 ////////////////////////////////////////////////////////////////////////////////////////////////////
 
-bool FMetadataFilterState::ApplyFilterToMetadata(TArrayView<const uint8>& Metadata) const
+bool FMetadataFilterState::ApplyFilterToMetadata(const TraceServices::FMetadataSpec* MetadataSpec, TArrayView<const uint8>& Metadata) const
 {
 	FMemoryReaderView MemoryReader(Metadata);
 	FCborReader CborReader(&MemoryReader, ECborEndianness::StandardCompliant);
 	FCborContext Context;
 
-	if (!CborReader.ReadNext(Context) || Context.MajorType() != ECborCode::Map)
+	ensure(SelectedDataType);
+
+	// We have 3 cases:
+	// 1. InMetadataSpec in not null and InMetadata is a list of values (dynamic metadata)
+	// 2. InMetadataSpec is null and InMetadata is a list of values (dynamic metadata without a spec emited or saved)
+	// 3. InMetadataSpec is null and InMetadata is a map of key value pairs. (static metadata)
+	bool bIsMap = false;
+	bool bHasReadFirstValue = false;
+	if (MetadataSpec == nullptr)
 	{
-		return false;
-	}
+		if (!CborReader.ReadNext(Context))
+		{
+			return false;
+		}
 
-	ensure(SelectedDataType);
+		if (Context.MajorType() == ECborCode::Map)
+		{
+			bIsMap = true;
+		}
+		else
+		{
+			bHasReadFirstValue = true;
+		}
+	}
 
-	while (true)
+	for (uint32 Index = 0; true; ++Index)
 	{
-		// Read key
-		if (!CborReader.ReadNext(Context) || !Context.IsString())
+		bool bKeyMatched = false;
+		if (MetadataSpec)
 		{
-			break;
+			TStringView<const TCHAR> CurrentKey;
+			if (Index < static_cast<uint32>(MetadataSpec->FieldNames.Num()))
+			{
+				CurrentKey = MetadataSpec->FieldNames[Index];
+			}
+
+			bKeyMatched = CurrentKey.Equals(Key, ESearchCase::IgnoreCase);
 		}
+		else if (bIsMap)
+		{
+			// Read key
+			if (!CborReader.ReadNext(Context) || !Context.IsString())
+			{
+				break;
+			}
 
-		FString CurrentKey = FString::ConstructFromPtrSize(Context.AsCString(), static_cast<int32>(Context.AsLength()));
+			TStringView<const ANSICHAR> CurrentKey(Context.AsCString(), static_cast<int32>(Context.AsLength()));
+			bKeyMatched = CurrentKey.Equals(Key, ESearchCase::IgnoreCase);
+		}
+		else
+		{
+			bKeyMatched = Key.IsEmpty();
+		}
 
-		if (!CborReader.ReadNext(Context))
+		if (Index > 0 || !bHasReadFirstValue)
 		{
-			break;
+			if (!CborReader.ReadNext(Context))
+			{
+				break;
+			}
 		}
 
-		if (Context.IsFiniteContainer())
+		if (Context.MajorType() == ECborCode::Array)
 		{
 			CborReader.SkipContainer(ECborCode::Array);
 		}
 
-		if (!CurrentKey.Equals(Key, ESearchCase::IgnoreCase))
+		if (!bKeyMatched && !bKeyAllowsAll)
 		{
 			continue;
 		}
 
+		if (bTermAllowsAll)
+		{
+			return true;
+		}
+
 		switch (Context.MajorType())
 		{
 		case ECborCode::Int:
@@ -461,7 +528,7 @@ bool FMetadataFilterState::ApplyFilterToMetadata(TArrayView<const uint8>& Metada
 			double InputValue = ConvertedData.Get<double>();
 
 			using FFilterOperatorDouble = FFilterOperator<double>;
-			FFilterOperatorDouble* Operator = (FFilterOperatorDouble*) SelectedOperator.Get();
+			FFilterOperatorDouble* Operator = (FFilterOperatorDouble*)SelectedOperator.Get();
 			if (Operator->Apply(Value, InputValue))
 			{
 				return true;
diff --git a/Engine/Source/Developer/TraceInsights/Private/Insights/TimingProfiler/ViewModels/TimerFilters.h b/Engine/Source/Developer/TraceInsights/Private/Insights/TimingProfiler/ViewModels/TimerFilters.h
--- a/Engine/Source/Developer/TraceInsights/Private/Insights/TimingProfiler/ViewModels/TimerFilters.h
+++ b/Engine/Source/Developer/TraceInsights/Private/Insights/TimingProfiler/ViewModels/TimerFilters.h
@@ -13,7 +13,7 @@
 
 namespace TraceServices
 {
-	class ITimingProfilerTimerReader;
+	struct FMetadataSpec;
 }
 
 namespace UE::Insights::TimingProfiler
@@ -138,7 +138,9 @@ private:
 	FText GetTermTextBoxValue() const;
 	void OnTermTextBoxValueCommitted(const FText& InNewText, ETextCommit::Type InTextCommit);
 
-	bool ApplyFilterToMetadata(TArrayView<const uint8>& Metadata) const;
+	bool ApplyFilterToMetadata(const TraceServices::FMetadataSpec* MetadataSpec, TArrayView<const uint8>& Metadata) const;
+
+	bool ShouldShowAllMetadataEvents() const { return bKeyAllowsAll && bTermAllowsAll; }
 
 private:
 	FString Key;
@@ -147,7 +149,8 @@ private:
 	typedef TVariant<double, int64, bool> ConvertedDataVariant;
 	ConvertedDataVariant ConvertedData;
 
-	bool bShowAllMetadataEvents = false;
+	bool bKeyAllowsAll = false;
+	bool bTermAllowsAll = false;
 
 	TArray<TSharedPtr<FMetadataFilterDataTypeEntry>> AvailableDataTypes;
 	TSharedPtr<FMetadataFilterDataTypeEntry> SelectedDataType;
diff --git a/Engine/Source/Developer/TraceInsightsCore/Private/InsightsCore/Filter/ViewModels/Filters.cpp b/Engine/Source/Developer/TraceInsightsCore/Private/InsightsCore/Filter/ViewModels/Filters.cpp
--- a/Engine/Source/Developer/TraceInsightsCore/Private/InsightsCore/Filter/ViewModels/Filters.cpp
+++ b/Engine/Source/Developer/TraceInsightsCore/Private/InsightsCore/Filter/ViewModels/Filters.cpp
@@ -45,10 +45,10 @@ FFilterStorage::FFilterStorage()
 	IntegerOperators->Add(MakeShared<FFilterOperator<int64>>(EFilterOperator::Gte, TEXT("\u2265"), std::greater_equal<>()));
 
 	StringOperators = MakeShared<TArray<TSharedPtr<IFilterOperator>>>();
-	StringOperators->Add(MakeShared<FFilterOperator<FStringView>>(EFilterOperator::Eq, TEXT("IS"), [](FStringView lhs, FStringView rhs) { return lhs.Equals(rhs); }));
-	StringOperators->Add(MakeShared<FFilterOperator<FStringView>>(EFilterOperator::NotEq, TEXT("IS NOT"), [](FStringView lhs, FStringView rhs) { return !lhs.Equals(rhs); }));
-	StringOperators->Add(MakeShared<FFilterOperator<FStringView>>(EFilterOperator::Contains, TEXT("CONTAINS"), [](FStringView lhs, FStringView rhs) { return lhs.Contains(rhs); }));
-	StringOperators->Add(MakeShared<FFilterOperator<FStringView>>(EFilterOperator::NotContains, TEXT("NOT CONTAINS"), [](FStringView lhs, FStringView rhs) { return !lhs.Contains(rhs); }));
+	StringOperators->Add(MakeShared<FFilterOperator<FStringView>>(EFilterOperator::Eq, TEXT("IS"), [](FStringView lhs, FStringView rhs) { return lhs.Equals(rhs, ESearchCase::IgnoreCase); }));
+	StringOperators->Add(MakeShared<FFilterOperator<FStringView>>(EFilterOperator::NotEq, TEXT("IS NOT"), [](FStringView lhs, FStringView rhs) { return !lhs.Equals(rhs, ESearchCase::IgnoreCase); }));
+	StringOperators->Add(MakeShared<FFilterOperator<FStringView>>(EFilterOperator::Contains, TEXT("CONTAINS"), [](FStringView lhs, FStringView rhs) { return lhs.Contains(rhs, ESearchCase::IgnoreCase); }));
+	StringOperators->Add(MakeShared<FFilterOperator<FStringView>>(EFilterOperator::NotContains, TEXT("NOT CONTAINS"), [](FStringView lhs, FStringView rhs) { return !lhs.Contains(rhs, ESearchCase::IgnoreCase); }));
 
 	FilterGroupOperators.Add(MakeShared<FFilterGroupOperator>(EFilterGroupOperator::And, LOCTEXT("AllOf", "All Of (AND)"), LOCTEXT("AllOfDesc", "All of the children must be true for the group to return true. Equivalent to an AND operation.")));
 	FilterGroupOperators.Add(MakeShared<FFilterGroupOperator>(EFilterGroupOperator::Or, LOCTEXT("AnyOf", "Any Of (OR)"), LOCTEXT("AnyOfDesc", "Any of the children must be true for the group to return true. Equivalent to an OR operation.")));
