[Insights] Asset Loading Insights: Minor refactoring of how tables are created to make code more readable and to avoid double call for ReadLoadTimeProfilerProvider.

#rnx

#Synced-CL 50769868

diff --git a/Engine/Source/Developer/TraceInsights/Private/Insights/LoadingProfiler/Widgets/SLoadingProfilerWindow.cpp b/Engine/Source/Developer/TraceInsights/Private/Insights/LoadingProfiler/Widgets/SLoadingProfilerWindow.cpp
--- a/Engine/Source/Developer/TraceInsights/Private/Insights/LoadingProfiler/Widgets/SLoadingProfilerWindow.cpp
+++ b/Engine/Source/Developer/TraceInsights/Private/Insights/LoadingProfiler/Widgets/SLoadingProfilerWindow.cpp
@@ -117,21 +117,26 @@ void SLoadingProfilerWindow::UpdateTableTreeViews()
 
 void SLoadingProfilerWindow::UpdateEventAggregationTreeView()
 {
+	using namespace TraceServices;
+
 	if (EventAggregationTreeView)
 	{
-		TSharedPtr<const TraceServices::IAnalysisSession> Session = FInsightsManager::Get()->GetSession();
-		if (Session.IsValid() && TraceServices::ReadLoadTimeProfilerProvider(*Session.Get()))
-		{
-			TraceServices::FAnalysisSessionReadScope SessionReadScope(*Session.Get());
-			const TraceServices::ILoadTimeProfilerProvider& LoadTimeProfilerProvider = *TraceServices::ReadLoadTimeProfilerProvider(*Session.Get());
+		TSharedPtr<IUntypedTable> SourceTable;
 
-			TraceServices::ITable<TraceServices::FLoadTimeProfilerAggregatedStats>* EventAggregationTable = LoadTimeProfilerProvider.CreateEventAggregation(SelectionStartTime, SelectionEndTime);
-			EventAggregationTreeView->UpdateSourceTable(MakeShareable(EventAggregationTable));
-		}
-		else
+		TSharedPtr<const IAnalysisSession> Session = FInsightsManager::Get()->GetSession();
+		if (Session.IsValid())
 		{
-			EventAggregationTreeView->UpdateSourceTable(nullptr);
+			const ILoadTimeProfilerProvider* LoadTimeProfilerProvider = ReadLoadTimeProfilerProvider(*Session.Get());
+			if (LoadTimeProfilerProvider)
+			{
+				FAnalysisSessionReadScope SessionReadScope(*Session.Get());
+				ITable<FLoadTimeProfilerAggregatedStats>* EventAggregationTable =
+					LoadTimeProfilerProvider->CreateEventAggregation(SelectionStartTime, SelectionEndTime);
+				SourceTable = MakeShareable(EventAggregationTable);
+			}
 		}
+
+		EventAggregationTreeView->UpdateSourceTable(SourceTable);
 	}
 }
 
@@ -139,21 +144,26 @@ void SLoadingProfilerWindow::UpdateEventAggregationTreeView()
 
 void SLoadingProfilerWindow::UpdateObjectTypeAggregationTreeView()
 {
+	using namespace TraceServices;
+
 	if (ObjectTypeAggregationTreeView)
 	{
-		TSharedPtr<const TraceServices::IAnalysisSession> Session = FInsightsManager::Get()->GetSession();
-		if (Session.IsValid() && TraceServices::ReadLoadTimeProfilerProvider(*Session.Get()))
-		{
-			TraceServices::FAnalysisSessionReadScope SessionReadScope(*Session.Get());
-			const TraceServices::ILoadTimeProfilerProvider& LoadTimeProfilerProvider = *TraceServices::ReadLoadTimeProfilerProvider(*Session.Get());
+		TSharedPtr<IUntypedTable> SourceTable;
 
-			TraceServices::ITable<TraceServices::FLoadTimeProfilerAggregatedStats>* ObjectTypeAggregationTable = LoadTimeProfilerProvider.CreateObjectTypeAggregation(SelectionStartTime, SelectionEndTime);
-			ObjectTypeAggregationTreeView->UpdateSourceTable(MakeShareable(ObjectTypeAggregationTable));
-		}
-		else
+		TSharedPtr<const IAnalysisSession> Session = FInsightsManager::Get()->GetSession();
+		if (Session.IsValid())
 		{
-			ObjectTypeAggregationTreeView->UpdateSourceTable(nullptr);
+			const ILoadTimeProfilerProvider* LoadTimeProfilerProvider = ReadLoadTimeProfilerProvider(*Session.Get());
+			if (LoadTimeProfilerProvider)
+			{
+				FAnalysisSessionReadScope SessionReadScope(*Session.Get());
+				ITable<FLoadTimeProfilerAggregatedStats>* ObjectTypeAggregationTable =
+					LoadTimeProfilerProvider->CreateObjectTypeAggregation(SelectionStartTime, SelectionEndTime);
+				SourceTable = MakeShareable(ObjectTypeAggregationTable);
+			}
 		}
+
+		ObjectTypeAggregationTreeView->UpdateSourceTable(SourceTable);
 	}
 }
 
@@ -161,21 +171,26 @@ void SLoadingProfilerWindow::UpdateObjectTypeAggregationTreeView()
 
 void SLoadingProfilerWindow::UpdatePackageDetailsTreeView()
 {
+	using namespace TraceServices;
+
 	if (PackageDetailsTreeView)
 	{
-		TSharedPtr<const TraceServices::IAnalysisSession> Session = FInsightsManager::Get()->GetSession();
-		if (Session.IsValid() && TraceServices::ReadLoadTimeProfilerProvider(*Session.Get()))
-		{
-			TraceServices::FAnalysisSessionReadScope SessionReadScope(*Session.Get());
-			const TraceServices::ILoadTimeProfilerProvider& LoadTimeProfilerProvider = *TraceServices::ReadLoadTimeProfilerProvider(*Session.Get());
+		TSharedPtr<IUntypedTable> SourceTable;
 
-			TraceServices::ITable<TraceServices::FPackagesTableRow>* PackageDetailsTable = LoadTimeProfilerProvider.CreatePackageDetailsTable(SelectionStartTime, SelectionEndTime);
-			PackageDetailsTreeView->UpdateSourceTable(MakeShareable(PackageDetailsTable));
-		}
-		else
+		TSharedPtr<const IAnalysisSession> Session = FInsightsManager::Get()->GetSession();
+		if (Session.IsValid())
 		{
-			PackageDetailsTreeView->UpdateSourceTable(nullptr);
+			const ILoadTimeProfilerProvider* LoadTimeProfilerProvider = ReadLoadTimeProfilerProvider(*Session.Get());
+			if (LoadTimeProfilerProvider)
+			{
+				FAnalysisSessionReadScope SessionReadScope(*Session.Get());
+				ITable<FPackagesTableRow>* PackageDetailsTable =
+					LoadTimeProfilerProvider->CreatePackageDetailsTable(SelectionStartTime, SelectionEndTime);
+				SourceTable = MakeShareable(PackageDetailsTable);
+			}
 		}
+
+		PackageDetailsTreeView->UpdateSourceTable(SourceTable);
 	}
 }
 
@@ -183,21 +198,26 @@ void SLoadingProfilerWindow::UpdatePackageDetailsTreeView()
 
 void SLoadingProfilerWindow::UpdateExportDetailsTreeView()
 {
+	using namespace TraceServices;
+
 	if (ExportDetailsTreeView)
 	{
-		TSharedPtr<const TraceServices::IAnalysisSession> Session = FInsightsManager::Get()->GetSession();
-		if (Session.IsValid() && TraceServices::ReadLoadTimeProfilerProvider(*Session.Get()))
-		{
-			TraceServices::FAnalysisSessionReadScope SessionReadScope(*Session.Get());
-			const TraceServices::ILoadTimeProfilerProvider& LoadTimeProfilerProvider = *TraceServices::ReadLoadTimeProfilerProvider(*Session.Get());
+		TSharedPtr<IUntypedTable> SourceTable;
 
-			TraceServices::ITable<TraceServices::FExportsTableRow>* ExportDetailsTable = LoadTimeProfilerProvider.CreateExportDetailsTable(SelectionStartTime, SelectionEndTime);
-			ExportDetailsTreeView->UpdateSourceTable(MakeShareable(ExportDetailsTable));
-		}
-		else
+		TSharedPtr<const IAnalysisSession> Session = FInsightsManager::Get()->GetSession();
+		if (Session.IsValid())
 		{
-			ExportDetailsTreeView->UpdateSourceTable(nullptr);
+			const ILoadTimeProfilerProvider* LoadTimeProfilerProvider = ReadLoadTimeProfilerProvider(*Session.Get());
+			if (LoadTimeProfilerProvider)
+			{
+				FAnalysisSessionReadScope SessionReadScope(*Session.Get());
+				ITable<FExportsTableRow>* ExportDetailsTable =
+					LoadTimeProfilerProvider->CreateExportDetailsTable(SelectionStartTime, SelectionEndTime);
+				SourceTable = MakeShareable(ExportDetailsTable);
+			}
 		}
+
+		ExportDetailsTreeView->UpdateSourceTable(SourceTable);
 	}
 }
 
@@ -205,21 +225,26 @@ void SLoadingProfilerWindow::UpdateExportDetailsTreeView()
 
 void SLoadingProfilerWindow::UpdateRequestsTreeView()
 {
+	using namespace TraceServices;
+
 	if (RequestsTreeView)
 	{
-		TSharedPtr<const TraceServices::IAnalysisSession> Session = FInsightsManager::Get()->GetSession();
-		if (Session.IsValid() && TraceServices::ReadLoadTimeProfilerProvider(*Session.Get()))
-		{
-			TraceServices::FAnalysisSessionReadScope SessionReadScope(*Session.Get());
-			const TraceServices::ILoadTimeProfilerProvider& LoadTimeProfilerProvider = *TraceServices::ReadLoadTimeProfilerProvider(*Session.Get());
+		TSharedPtr<IUntypedTable> SourceTable;
 
-			TraceServices::ITable<TraceServices::FRequestsTableRow>* RequestsTable = LoadTimeProfilerProvider.CreateRequestsTable(SelectionStartTime, SelectionEndTime);
-			RequestsTreeView->UpdateSourceTable(MakeShareable(RequestsTable));
-		}
-		else
+		TSharedPtr<const IAnalysisSession> Session = FInsightsManager::Get()->GetSession();
+		if (Session.IsValid())
 		{
-			RequestsTreeView->UpdateSourceTable(nullptr);
+			const ILoadTimeProfilerProvider* LoadTimeProfilerProvider = ReadLoadTimeProfilerProvider(*Session.Get());
+			if (LoadTimeProfilerProvider)
+			{
+				FAnalysisSessionReadScope SessionReadScope(*Session.Get());
+				ITable<FRequestsTableRow>* RequestsTable =
+					LoadTimeProfilerProvider->CreateRequestsTable(SelectionStartTime, SelectionEndTime);
+				SourceTable = MakeShareable(RequestsTable);
+			}
 		}
+
+		RequestsTreeView->UpdateSourceTable(SourceTable);
 	}
 }
 
