[Insights] Added enum ETimerAggregationMode for the supported aggregation modes in the Timers panel, in preparation for adding more frame types to the ETraceFrameType enum.

#Synced-CL 50972180

diff --git a/Engine/Source/Developer/TraceInsights/Private/Insights/InsightsSettings.cpp b/Engine/Source/Developer/TraceInsights/Private/Insights/InsightsSettings.cpp
--- a/Engine/Source/Developer/TraceInsights/Private/Insights/InsightsSettings.cpp
+++ b/Engine/Source/Developer/TraceInsights/Private/Insights/InsightsSettings.cpp
@@ -223,7 +223,7 @@ void FInsightsSettings::LoadFromConfig()
 	GConfig->GetArray(TEXT("Insights.TimingProfiler.TimersView"), TEXT("GameFrameColumns"), TimersViewGameFrameVisibleColumns, SettingsIni);
 	GConfig->GetArray(TEXT("Insights.TimingProfiler.TimersView"), TEXT("RenderingFrameColumns"), TimersViewRenderingFrameVisibleColumns, SettingsIni);
 
-	GConfig->GetInt(TEXT("Insights.TimingProfiler.TimersView"), TEXT("Mode"), TimersViewMode, SettingsIni);
+	GConfig->GetInt(TEXT("Insights.TimingProfiler.TimersView"), TEXT("AggregationMode"), TimersViewAggregationMode, SettingsIni);
 
 	GConfig->GetInt(TEXT("Insights.TimingProfiler.TimersView"), TEXT("GroupingMode"), TimersViewGroupingMode, SettingsIni);
 	GConfig->GetBool(TEXT("Insights.TimingProfiler.TimersView"), TEXT("ShowGpuTimers"), bTimersViewShowGpuTimers, SettingsIni);
@@ -362,7 +362,7 @@ void FInsightsSettings::SaveToConfig()
 	GConfig->SetArray(TEXT("Insights.TimingProfiler.TimersView"), TEXT("GameFrameColumns"), TimersViewGameFrameVisibleColumns, SettingsIni);
 	GConfig->SetArray(TEXT("Insights.TimingProfiler.TimersView"), TEXT("RenderingFrameColumns"), TimersViewRenderingFrameVisibleColumns, SettingsIni);
 
-	GConfig->SetInt(TEXT("Insights.TimingProfiler.TimersView"), TEXT("Mode"), TimersViewMode, SettingsIni);
+	GConfig->SetInt(TEXT("Insights.TimingProfiler.TimersView"), TEXT("AggregationMode"), TimersViewAggregationMode, SettingsIni);
 
 	GConfig->SetInt(TEXT("Insights.TimingProfiler.TimersView"), TEXT("GroupingMode"), TimersViewGroupingMode, SettingsIni);
 	GConfig->SetBool(TEXT("Insights.TimingProfiler.TimersView"), TEXT("ShowGpuTimers"), bTimersViewShowGpuTimers, SettingsIni);
diff --git a/Engine/Source/Developer/TraceInsights/Private/Insights/InsightsSettings.h b/Engine/Source/Developer/TraceInsights/Private/Insights/InsightsSettings.h
--- a/Engine/Source/Developer/TraceInsights/Private/Insights/InsightsSettings.h
+++ b/Engine/Source/Developer/TraceInsights/Private/Insights/InsightsSettings.h
@@ -484,9 +484,9 @@ public:
 	void SetTimersViewRenderingFrameVisibleColumns(const TArray<FString>& Columns) { TimersViewRenderingFrameVisibleColumns = Columns; }
 	void SetAndSaveTimersViewRenderingFrameVisibleColumns(const TArray<FString>& Columns) { SET_AND_SAVE(TimersViewRenderingFrameVisibleColumns, Columns); }
 
-	int32 GetTimersViewMode() const { return TimersViewMode; }
-	void SetTimersViewMode(int32 InMode) { TimersViewMode = InMode; }
-	void SetAndSaveTimersViewMode(int32 InMode) { SET_AND_SAVE(TimersViewMode, InMode); }
+	int32 GetTimersViewAggregationMode() const { return TimersViewAggregationMode; }
+	void SetTimersViewAggregationMode(int32 InMode) { TimersViewAggregationMode = InMode; }
+	void SetAndSaveTimersViewAggregationMode(int32 InMode) { SET_AND_SAVE(TimersViewAggregationMode, InMode); }
 
 	int32 GetTimersViewGroupingMode() const { return TimersViewGroupingMode; }
 	void SetTimersViewGroupingMode(int32 InValue) { TimersViewGroupingMode = InValue; }
@@ -656,9 +656,9 @@ private:
 
 	/**
 	 * The mode for the timers panel.
-	 * See ETraceFrameType in MiscTrace.h.
+	 * See ETimerAggregationMode in TimingProfiler\ViewModels\TimerAggregator.h.
 	 */
-	int32 TimersViewMode = 2; // (int32)TraceFrameType_Count
+	int32 TimersViewAggregationMode = 0; // 0=Instance, 1=GameFrame, 2=RenderingFrame
 
 	/** The grouping mode for the timers panel. */
 	int32 TimersViewGroupingMode = 3; // ByType
diff --git a/Engine/Source/Developer/TraceInsights/Private/Insights/TimingProfiler/ViewModels/TimerAggregator.cpp b/Engine/Source/Developer/TraceInsights/Private/Insights/TimingProfiler/ViewModels/TimerAggregator.cpp
--- a/Engine/Source/Developer/TraceInsights/Private/Insights/TimingProfiler/ViewModels/TimerAggregator.cpp
+++ b/Engine/Source/Developer/TraceInsights/Private/Insights/TimingProfiler/ViewModels/TimerAggregator.cpp
@@ -2,6 +2,8 @@
 
 #include "TimerAggregator.h"
 
+#include "ProfilingDebugging/MiscTrace.h" // for ETraceFrameType
+
 // TraceServices
 #include "Common/ProviderLock.h"
 #include "TraceServices/Model/AnalysisSession.h"
@@ -30,7 +32,7 @@ struct FTimerAggregationParams
 	TSet<uint32> CpuThreads;  // [Thread Id]
 	TSet<uint32> CpuSamplingThreads; // [System Thread Id]
 
-	ETraceFrameType FrameType = ETraceFrameType::TraceFrameType_Count;
+	ETimerAggregationMode AggregationMode = ETimerAggregationMode::Instance;
 };
 
 ////////////////////////////////////////////////////////////////////////////////////////////////////
@@ -50,17 +52,41 @@ public:
 		, bIncludeVerseSampling(InParams.bIncludeVerseSampling)
 		, CpuThreads(InParams.CpuThreads)
 		, CpuSamplingThreads(InParams.CpuSamplingThreads)
-		, FrameType(InParams.FrameType)
+		, AggregationMode(InParams.AggregationMode)
 		, ResultTable()
 	{
 	}
 
-	virtual ~FTimerAggregationWorker() {}
+	virtual ~FTimerAggregationWorker()
+	{
+	}
 
 	virtual void DoWork(TSharedPtr<TraceServices::FCancellationToken> CancellationToken) override;
 
-	TraceServices::ITable<TraceServices::FTimingProfilerAggregatedStats>* GetResultTable() const { return ResultTable.Get(); }
-	void ResetResults() { ResultTable.Reset(); }
+	TraceServices::ITable<TraceServices::FTimingProfilerAggregatedStats>* GetResultTable() const
+	{
+		return ResultTable.Get();
+	}
+
+	void ResetResults()
+	{
+		ResultTable.Reset();
+	}
+
+	ETraceFrameType GetFrameType() const
+	{
+		switch (AggregationMode)
+		{
+		case ETimerAggregationMode::Instance:
+			return ETraceFrameType::TraceFrameType_Count;
+		case ETimerAggregationMode::GameFrame:
+			return ETraceFrameType::TraceFrameType_Game;
+		case ETimerAggregationMode::RenderingFrame:
+			return ETraceFrameType::TraceFrameType_Rendering;
+		default:
+			return ETraceFrameType::TraceFrameType_Count;
+		}
+	}
 
 private:
 	TSharedPtr<const TraceServices::IAnalysisSession> Session;
@@ -72,7 +98,7 @@ private:
 	bool bIncludeVerseSampling;
 	TSet<uint32> CpuThreads; // [Thread Id]
 	TSet<uint32> CpuSamplingThreads; // [System Thread Id]
-	ETraceFrameType FrameType;
+	ETimerAggregationMode AggregationMode;
 	TUniquePtr<TraceServices::ITable<TraceServices::FTimingProfilerAggregatedStats>> ResultTable;
 };
 
@@ -99,7 +125,7 @@ void FTimerAggregationWorker::DoWork(TSharedPtr<TraceServices::FCancellationToke
 	Params.bIncludeOldGpu2 = bIncludeOldGpu2;
 	Params.bIncludeVerseSampling = bIncludeVerseSampling;
 	Params.CpuThreadFilter = [this](uint32 ThreadId) { return CpuThreads.Contains(ThreadId); };
-	Params.FrameType = FrameType;
+	Params.FrameType = GetFrameType();
 	Params.CancellationToken = CancellationToken;
 
 	bool bHasCpuSamplingTimelines = false;
@@ -143,7 +169,7 @@ IStatsAggregationWorker* FTimerAggregator::CreateWorker(TSharedPtr<const TraceSe
 
 	Params.StartTime = GetIntervalStartTime();
 	Params.EndTime = GetIntervalEndTime();
-	Params.FrameType = GetFrameType();
+	Params.AggregationMode = GetAggregationMode();
 
 	TSharedPtr<STimingProfilerWindow> Wnd = FTimingProfilerManager::Get()->GetProfilerWindow();
 	if (Wnd.IsValid())
diff --git a/Engine/Source/Developer/TraceInsights/Private/Insights/TimingProfiler/ViewModels/TimerAggregator.h b/Engine/Source/Developer/TraceInsights/Private/Insights/TimingProfiler/ViewModels/TimerAggregator.h
--- a/Engine/Source/Developer/TraceInsights/Private/Insights/TimingProfiler/ViewModels/TimerAggregator.h
+++ b/Engine/Source/Developer/TraceInsights/Private/Insights/TimingProfiler/ViewModels/TimerAggregator.h
@@ -4,8 +4,6 @@
 
 #include "CoreTypes.h"
 
-#include "ProfilingDebugging/MiscTrace.h" // for ETraceFrameType
-
 // TraceServices
 #include "TraceServices/Containers/Tables.h"
 #include "TraceServices/Model/TimingProfiler.h"
@@ -20,6 +18,17 @@ namespace UE::Insights::TimingProfiler
 
 ////////////////////////////////////////////////////////////////////////////////////////////////////
 
+enum class ETimerAggregationMode
+{
+	Instance,
+	GameFrame,
+	RenderingFrame,
+
+	Count
+};
+
+////////////////////////////////////////////////////////////////////////////////////////////////////
+
 class FTimerAggregator : public FStatsAggregator
 {
 public:
@@ -29,14 +38,14 @@ public:
 	TraceServices::ITable<TraceServices::FTimingProfilerAggregatedStats>* GetResultTable() const;
 	void ResetResults();
 
-	ETraceFrameType GetFrameType() { return FrameType; }
-	void SetFrameType(ETraceFrameType InFrameType) { FrameType = InFrameType; }
+	ETimerAggregationMode GetAggregationMode() const { return AggregationMode; }
+	void SetAggregationMode(ETimerAggregationMode InMode) { AggregationMode = InMode; }
 
 protected:
 	virtual IStatsAggregationWorker* CreateWorker(TSharedPtr<const TraceServices::IAnalysisSession> InSession) override;
 
 private:
-	ETraceFrameType FrameType = ETraceFrameType::TraceFrameType_Count;
+	ETimerAggregationMode AggregationMode = ETimerAggregationMode::Instance;
 };
 
 ////////////////////////////////////////////////////////////////////////////////////////////////////
diff --git a/Engine/Source/Developer/TraceInsights/Private/Insights/TimingProfiler/ViewModels/TimersViewColumnFactory.cpp b/Engine/Source/Developer/TraceInsights/Private/Insights/TimingProfiler/ViewModels/TimersViewColumnFactory.cpp
--- a/Engine/Source/Developer/TraceInsights/Private/Insights/TimingProfiler/ViewModels/TimersViewColumnFactory.cpp
+++ b/Engine/Source/Developer/TraceInsights/Private/Insights/TimingProfiler/ViewModels/TimersViewColumnFactory.cpp
@@ -127,9 +127,9 @@ TSharedRef<FTableColumn> FTimersViewColumnFactory::CreateNameColumn()
 	Column.SetTitleName(LOCTEXT("Name_ColumnTitle", "Timer or Group Name"));
 
 	FText Description = LOCTEXT("Name_ColumnDesc", "Name of timer or group");
-	Column.SetDescription(ETraceFrameType::TraceFrameType_Count, Description);
-	Column.SetDescription(ETraceFrameType::TraceFrameType_Game, Description);
-	Column.SetDescription(ETraceFrameType::TraceFrameType_Rendering, Description);
+	Column.SetDescription(ETimerAggregationMode::Instance, Description);
+	Column.SetDescription(ETimerAggregationMode::GameFrame, Description);
+	Column.SetDescription(ETimerAggregationMode::RenderingFrame, Description);
 
 	Column.SetFlags(ETableColumnFlags::ShouldBeVisible |
 					ETableColumnFlags::CanBeFiltered |
@@ -164,9 +164,9 @@ TSharedRef<FTableColumn> FTimersViewColumnFactory::CreateIdColumn()
 	Column.SetTitleName(LOCTEXT("Id_ColumnTitle", "Timer Id"));
 
 	FText Description = LOCTEXT("Id_ColumnDesc", "Temporary index/id assigned to timers during analysis");
-	Column.SetDescription(ETraceFrameType::TraceFrameType_Count, Description);
-	Column.SetDescription(ETraceFrameType::TraceFrameType_Game, Description);
-	Column.SetDescription(ETraceFrameType::TraceFrameType_Rendering, Description);
+	Column.SetDescription(ETimerAggregationMode::Instance, Description);
+	Column.SetDescription(ETimerAggregationMode::GameFrame, Description);
+	Column.SetDescription(ETimerAggregationMode::RenderingFrame, Description);
 
 	Column.SetFlags(ETableColumnFlags::CanBeHidden |
 					ETableColumnFlags::CanBeFiltered);
@@ -219,9 +219,9 @@ TSharedRef<FTableColumn> FTimersViewColumnFactory::CreateMetaGroupNameColumn()
 	Column.SetTitleName(LOCTEXT("MetaGroupName_ColumnTitle", "Meta Group Name"));
 
 	FText Description = LOCTEXT("MetaGroupName_ColumnDesc", "Name of the meta group");
-	Column.SetDescription(ETraceFrameType::TraceFrameType_Count, Description);
-	Column.SetDescription(ETraceFrameType::TraceFrameType_Game, Description);
-	Column.SetDescription(ETraceFrameType::TraceFrameType_Rendering, Description);
+	Column.SetDescription(ETimerAggregationMode::Instance, Description);
+	Column.SetDescription(ETimerAggregationMode::GameFrame, Description);
+	Column.SetDescription(ETimerAggregationMode::RenderingFrame, Description);
 
 	Column.SetFlags(ETableColumnFlags::CanBeHidden |
 					//ETableColumnFlags::ShouldBeVisible |
@@ -266,9 +266,9 @@ TSharedRef<FTableColumn> FTimersViewColumnFactory::CreateTypeColumn()
 	Column.SetTitleName(LOCTEXT("Type_ColumnTitle", "Type"));
 
 	FText Description = LOCTEXT("Type_ColumnDesc", "Type of timer or group");
-	Column.SetDescription(ETraceFrameType::TraceFrameType_Count, Description);
-	Column.SetDescription(ETraceFrameType::TraceFrameType_Game, Description);
-	Column.SetDescription(ETraceFrameType::TraceFrameType_Rendering, Description);
+	Column.SetDescription(ETimerAggregationMode::Instance, Description);
+	Column.SetDescription(ETimerAggregationMode::GameFrame, Description);
+	Column.SetDescription(ETimerAggregationMode::RenderingFrame, Description);
 
 	Column.SetFlags(ETableColumnFlags::CanBeHidden |
 					//ETableColumnFlags::ShouldBeVisible |
@@ -314,9 +314,9 @@ TSharedRef<FTableColumn> FTimersViewColumnFactory::CreateInstanceCountColumn()
 	Column.SetTitleName(LOCTEXT("InstanceCount_ColumnTitle", "Instance Count"));
 
 	FText Description = LOCTEXT("InstanceCount_ColumnDesc", "Number of timing event instances of the selected timer");
-	Column.SetDescription(ETraceFrameType::TraceFrameType_Count, Description);
-	Column.SetDescription(ETraceFrameType::TraceFrameType_Game, Description);
-	Column.SetDescription(ETraceFrameType::TraceFrameType_Rendering, Description);
+	Column.SetDescription(ETimerAggregationMode::Instance, Description);
+	Column.SetDescription(ETimerAggregationMode::GameFrame, Description);
+	Column.SetDescription(ETimerAggregationMode::RenderingFrame, Description);
 
 	Column.SetFlags(ETableColumnFlags::CanBeHidden |
 					ETableColumnFlags::ShouldBeVisible |
@@ -363,9 +363,9 @@ TSharedRef<FTableColumn> FTimersViewColumnFactory::CreateChildInstanceCountColum
 	Column.SetTitleName(LOCTEXT("ChildInstanceCount_ColumnTitle", "Child Instance Count"));
 
 	FText Description = LOCTEXT("ChildInstanceCount_ColumnDesc", "Total number of timing event instances of the child timers (callers or callees)");
-	Column.SetDescription(ETraceFrameType::TraceFrameType_Count, Description);
-	Column.SetDescription(ETraceFrameType::TraceFrameType_Game, Description);
-	Column.SetDescription(ETraceFrameType::TraceFrameType_Rendering, Description);
+	Column.SetDescription(ETimerAggregationMode::Instance, Description);
+	Column.SetDescription(ETimerAggregationMode::GameFrame, Description);
+	Column.SetDescription(ETimerAggregationMode::RenderingFrame, Description);
 
 	Column.SetFlags(ETableColumnFlags::CanBeHidden | ETableColumnFlags::CanBeFiltered);
 
@@ -414,15 +414,15 @@ TSharedRef<FTableColumn> FTimersViewColumnFactory::CreateAverageInstanceCountCol
 	Column.SetShortName(LOCTEXT("AverageInstanceCount_ColumnName", "C.Avg"));
 	Column.SetTitleName(LOCTEXT("AverageInstanceCount_ColumnTitle", "Average Instance Count"));
 
-	Column.SetDescription(ETraceFrameType::TraceFrameType_Count,
+	Column.SetDescription(ETimerAggregationMode::Instance,
 		LOCTEXT("AverageInstanceCount_ColumnDesc",
 			"Total number of timing event instances of the selected timer in the selected range\n"
 			"In Instance mode, this column has the same value as the Instance Count column."));
-	Column.SetDescription(ETraceFrameType::TraceFrameType_Game,
+	Column.SetDescription(ETimerAggregationMode::GameFrame,
 		LOCTEXT("AverageInstanceCount_GameFrameColumnDesc",
 			"Game Frame Average Instance Count\n"
 			"Average number of timing event instances of the selected timer per-frame."));
-	Column.SetDescription(ETraceFrameType::TraceFrameType_Rendering,
+	Column.SetDescription(ETimerAggregationMode::RenderingFrame,
 		LOCTEXT("AverageInstanceCount_RenderingFrameColumnDesc",
 			"Rendering Frame Average Instance Count\n"
 			"Average number of timing event instances of the selected timer per-frame."));
@@ -472,9 +472,9 @@ TSharedRef<FTableColumn> FTimersViewColumnFactory::CreateTotalInclusiveTimeColum
 	Column.SetTitleName(LOCTEXT("TotalInclusiveTime_ColumnTitle", "Total Inclusive Time"));
 
 	FText Description = LOCTEXT("TotalInclusiveTime_ColumnDesc", "Total inclusive duration of selected timer's instances");
-	Column.SetDescription(ETraceFrameType::TraceFrameType_Count, Description);
-	Column.SetDescription(ETraceFrameType::TraceFrameType_Game, Description);
-	Column.SetDescription(ETraceFrameType::TraceFrameType_Rendering, Description);
+	Column.SetDescription(ETimerAggregationMode::Instance, Description);
+	Column.SetDescription(ETimerAggregationMode::GameFrame, Description);
+	Column.SetDescription(ETimerAggregationMode::RenderingFrame, Description);
 
 	Column.SetFlags(ETableColumnFlags::CanBeHidden |
 					ETableColumnFlags::ShouldBeVisible |
@@ -522,9 +522,9 @@ TSharedRef<FTableColumn> FTimersViewColumnFactory::CreateParentPercentageInclusi
 	Column.SetTitleName(LOCTEXT("ParentPercentageInclusiveTime_ColumnTitle", "Total Inclusive Time / Parent's Total Inclusive Time (%)"));
 
 	FText Description = LOCTEXT("ParentPercentageInclusiveTime_ColumnDesc", "Percentage of total inclusive duration (the selected timer relative to the parent timer)");
-	Column.SetDescription(ETraceFrameType::TraceFrameType_Count, Description);
-	Column.SetDescription(ETraceFrameType::TraceFrameType_Game, Description);
-	Column.SetDescription(ETraceFrameType::TraceFrameType_Rendering, Description);
+	Column.SetDescription(ETimerAggregationMode::Instance, Description);
+	Column.SetDescription(ETimerAggregationMode::GameFrame, Description);
+	Column.SetDescription(ETimerAggregationMode::RenderingFrame, Description);
 
 	Column.SetFlags(ETableColumnFlags::CanBeHidden |
 		ETableColumnFlags::ShouldBeVisible |
@@ -585,9 +585,9 @@ TSharedRef<FTableColumn> FTimersViewColumnFactory::CreateRootPercentageInclusive
 	Column.SetTitleName(LOCTEXT("RootPercentageInclusiveTime_ColumnTitle", "Total Inclusive Time / Root's Total Inclusive Time (%)"));
 
 	FText Description = LOCTEXT("RootPercentageInclusiveTime_ColumnDesc", "Percentage of total inclusive duration (the selected timer relative to the root timer)");
-	Column.SetDescription(ETraceFrameType::TraceFrameType_Count, Description);
-	Column.SetDescription(ETraceFrameType::TraceFrameType_Game, Description);
-	Column.SetDescription(ETraceFrameType::TraceFrameType_Rendering, Description);
+	Column.SetDescription(ETimerAggregationMode::Instance, Description);
+	Column.SetDescription(ETimerAggregationMode::GameFrame, Description);
+	Column.SetDescription(ETimerAggregationMode::RenderingFrame, Description);
 
 	Column.SetFlags(ETableColumnFlags::CanBeHidden |
 		ETableColumnFlags::ShouldBeVisible |
@@ -651,16 +651,16 @@ TSharedRef<FTableColumn> FTimersViewColumnFactory::CreateMaxInclusiveTimeColumn(
 	Column.SetShortName(LOCTEXT("MaxInclusiveTime_ColumnName", "I.Max"));
 	Column.SetTitleName(LOCTEXT("MaxInclusiveTime_ColumnTitle", "Max Inclusive Time (ms)"));
 
-	Column.SetDescription(ETraceFrameType::TraceFrameType_Count,
+	Column.SetDescription(ETimerAggregationMode::Instance,
 		LOCTEXT("MaxInclusiveTime_ColumnDesc",
 			"Maximum inclusive duration of selected timer's instances, in milliseconds"));
-	Column.SetDescription(ETraceFrameType::TraceFrameType_Game,
+	Column.SetDescription(ETimerAggregationMode::GameFrame,
 		LOCTEXT("MaxInclusiveTime_GameFrameColumnDesc",
 			"Game Frame Maximum Inclusive Duration\n"
 			"Inclusive duration is computed for a single frame as the sum of inclusive duration of all instances of the timer in the respective frame.\n"
 			"The maximum is selected from these per-frame inclusive durations.\n"
 			"Unit is milliseconds."));
-	Column.SetDescription(ETraceFrameType::TraceFrameType_Rendering,
+	Column.SetDescription(ETimerAggregationMode::RenderingFrame,
 		LOCTEXT("MaxInclusiveTime_RenderingFrameColumnDesc",
 			"Rendering Frame Maximum Inclusive Duration\n"
 			"Inclusive duration is computed for a single frame as the sum of inclusive duration of all instances of the timer in the respective frame.\n"
@@ -710,16 +710,16 @@ TSharedRef<FTableColumn> FTimersViewColumnFactory::CreateAverageInclusiveTimeCol
 	Column.SetShortName(LOCTEXT("AvgInclusiveTime_ColumnName", "I.Avg"));
 	Column.SetTitleName(LOCTEXT("AvgInclusiveTime_ColumnTitle", "Average Inclusive Time (ms)"));
 
-	Column.SetDescription(ETraceFrameType::TraceFrameType_Count,
+	Column.SetDescription(ETimerAggregationMode::Instance,
 		LOCTEXT("AvgInclusiveTime_ColumnDesc",
 			"Average inclusive duration of selected timer's instances, in milliseconds"));
-	Column.SetDescription(ETraceFrameType::TraceFrameType_Game,
+	Column.SetDescription(ETimerAggregationMode::GameFrame,
 		LOCTEXT("AvgInclusiveTime_GameFrameColumnDesc",
 			"Game Frame Average Inclusive Duration\n"
 			"Inclusive duration is computed for a single frame as the sum of inclusive duration of all instances of the timer in the respective frame.\n"
 			"The average is computed from these per-frame inclusive durations.\n"
 			"Unit is milliseconds."));
-	Column.SetDescription(ETraceFrameType::TraceFrameType_Rendering,
+	Column.SetDescription(ETimerAggregationMode::RenderingFrame,
 		LOCTEXT("AvgInclusiveTime_RenderingFrameColumnDesc",
 			"Rendering Frame Average Inclusive Duration\n"
 			"Inclusive duration is computed for a single frame as the sum of inclusive duration of all instances of the timer in the respective frame.\n"
@@ -769,16 +769,16 @@ TSharedRef<FTableColumn> FTimersViewColumnFactory::CreateMedianInclusiveTimeColu
 	Column.SetShortName(LOCTEXT("MedInclusiveTime_ColumnName", "I.Med"));
 	Column.SetTitleName(LOCTEXT("MedInclusiveTime_ColumnTitle", "Median Inclusive Time (ms)"));
 
-	Column.SetDescription(ETraceFrameType::TraceFrameType_Count,
+	Column.SetDescription(ETimerAggregationMode::Instance,
 		LOCTEXT("MedInclusiveTime_ColumnDesc",
 			"Median inclusive duration of selected timer's instances, in milliseconds"));
-	Column.SetDescription(ETraceFrameType::TraceFrameType_Game,
+	Column.SetDescription(ETimerAggregationMode::GameFrame,
 		LOCTEXT("MedInclusiveTime_GameFrameColumnDesc",
 			"Game Frame Median Inclusive Duration\n"
 			"Inclusive duration is computed for a single frame as the sum of inclusive duration of all instances of the timer in the respective frame.\n"
 			"The median is computed from these per-frame inclusive durations.\n"
 			"Unit is milliseconds."));
-	Column.SetDescription(ETraceFrameType::TraceFrameType_Rendering,
+	Column.SetDescription(ETimerAggregationMode::RenderingFrame,
 		LOCTEXT("MedInclusiveTime_RenderingFrameColumnDesc",
 			"Rendering Frame Median Inclusive Duration\n"
 			"Inclusive duration is computed for a single frame as the sum of inclusive duration of all instances of the timer in the respective frame.\n"
@@ -828,16 +828,16 @@ TSharedRef<FTableColumn> FTimersViewColumnFactory::CreateMinInclusiveTimeColumn(
 	Column.SetShortName(LOCTEXT("MinInclusiveTime_ColumnName", "I.Min"));
 	Column.SetTitleName(LOCTEXT("MinInclusiveTime_ColumnTitle", "Min Inclusive Time (ms)"));
 
-	Column.SetDescription(ETraceFrameType::TraceFrameType_Count,
+	Column.SetDescription(ETimerAggregationMode::Instance,
 		LOCTEXT("MinInclusiveTime_ColumnDesc",
 			"Minimum inclusive duration of selected timer's instances, in milliseconds"));
-	Column.SetDescription(ETraceFrameType::TraceFrameType_Game,
+	Column.SetDescription(ETimerAggregationMode::GameFrame,
 		LOCTEXT("MinInclusiveTime_GameFrameColumnDesc",
 			"Game Frame Minimum Inclusive Duration\n"
 			"Inclusive duration is computed for a single frame as the sum of inclusive duration of all instances of the timer in the respective frame.\n"
 			"The minimum is selected from these per-frame inclusive durations.\n"
 			"Unit is milliseconds."));
-	Column.SetDescription(ETraceFrameType::TraceFrameType_Rendering,
+	Column.SetDescription(ETimerAggregationMode::RenderingFrame,
 		LOCTEXT("MinInclusiveTime_RenderingFrameColumnDesc",
 			"Rendering Frame Minimum Inclusive Duration\n"
 			"Inclusive duration is computed for a single frame as the sum of inclusive duration of all instances of the timer in the respective frame.\n"
@@ -890,9 +890,9 @@ TSharedRef<FTableColumn> FTimersViewColumnFactory::CreateTotalExclusiveTimeColum
 	Column.SetTitleName(LOCTEXT("TotalExclusiveTime_ColumnTitle", "Total Exclusive Time"));
 
 	FText Description = LOCTEXT("TotalExclusiveTime_ColumnDesc", "Total exclusive duration of selected timer's instances");
-	Column.SetDescription(ETraceFrameType::TraceFrameType_Count, Description);
-	Column.SetDescription(ETraceFrameType::TraceFrameType_Game, Description);
-	Column.SetDescription(ETraceFrameType::TraceFrameType_Rendering, Description);
+	Column.SetDescription(ETimerAggregationMode::Instance, Description);
+	Column.SetDescription(ETimerAggregationMode::GameFrame, Description);
+	Column.SetDescription(ETimerAggregationMode::RenderingFrame, Description);
 
 	Column.SetFlags(ETableColumnFlags::CanBeHidden |
 					ETableColumnFlags::ShouldBeVisible |
@@ -940,9 +940,9 @@ TSharedRef<FTableColumn> FTimersViewColumnFactory::CreatePercentageExclusiveIncl
 	Column.SetTitleName(LOCTEXT("PercentageExclusiveInclusiveTime_ColumnTitle", "Total Exclusive Time / Total Inclusive Time (%)"));
 
 	FText Description = LOCTEXT("PercentageExclusiveInclusiveTime_ColumnDesc", "Percentage of total exclusive duration from the total inclusive duration");
-	Column.SetDescription(ETraceFrameType::TraceFrameType_Count, Description);
-	Column.SetDescription(ETraceFrameType::TraceFrameType_Game, Description);
-	Column.SetDescription(ETraceFrameType::TraceFrameType_Rendering, Description);
+	Column.SetDescription(ETimerAggregationMode::Instance, Description);
+	Column.SetDescription(ETimerAggregationMode::GameFrame, Description);
+	Column.SetDescription(ETimerAggregationMode::RenderingFrame, Description);
 
 	Column.SetFlags(ETableColumnFlags::CanBeHidden |
 		//ETableColumnFlags::ShouldBeVisible |
@@ -998,16 +998,16 @@ TSharedRef<FTableColumn> FTimersViewColumnFactory::CreateMaxExclusiveTimeColumn(
 	Column.SetShortName(LOCTEXT("MaxExclusiveTime_ColumnName", "E.Max"));
 	Column.SetTitleName(LOCTEXT("MaxExclusiveTime_ColumnTitle", "Max Exclusive Time (ms)"));
 
-	Column.SetDescription(ETraceFrameType::TraceFrameType_Count,
+	Column.SetDescription(ETimerAggregationMode::Instance,
 		LOCTEXT("MaxExclusiveTime_ColumnDesc",
 			"Maximum exclusive duration of selected timer's instances, in milliseconds"));
-	Column.SetDescription(ETraceFrameType::TraceFrameType_Game,
+	Column.SetDescription(ETimerAggregationMode::GameFrame,
 		LOCTEXT("MaxExclusiveTime_GameFrameColumnDesc",
 			"Game Frame Maximum Exclusive Duration\n"
 			"Exclusive duration is computed for a single frame as the sum of exclusive duration of all instances of the timer in the respective frame.\n"
 			"The maximum is selected from these per-frame exclusive durations.\n"
 			"Unit is milliseconds."));
-	Column.SetDescription(ETraceFrameType::TraceFrameType_Rendering,
+	Column.SetDescription(ETimerAggregationMode::RenderingFrame,
 		LOCTEXT("MaxExclusiveTime_RenderingFrameColumnDesc",
 			"Rendering Frame Maximum Exclusive Duration\n"
 			"Exclusive duration is computed for a single frame as the sum of exclusive duration of all instances of the timer in the respective frame.\n"
@@ -1057,16 +1057,16 @@ TSharedRef<FTableColumn> FTimersViewColumnFactory::CreateAverageExclusiveTimeCol
 	Column.SetShortName(LOCTEXT("AvgExclusiveTime_ColumnName", "E.Avg"));
 	Column.SetTitleName(LOCTEXT("AvgExclusiveTime_ColumnTitle", "Average Exclusive Time (ms)"));
 
-	Column.SetDescription(ETraceFrameType::TraceFrameType_Count,
+	Column.SetDescription(ETimerAggregationMode::Instance,
 		LOCTEXT("AvgExclusiveTime_ColumnDesc",
 			"Average exclusive duration of selected timer's instances, in milliseconds"));
-	Column.SetDescription(ETraceFrameType::TraceFrameType_Game,
+	Column.SetDescription(ETimerAggregationMode::GameFrame,
 		LOCTEXT("AvgExclusiveTime_GameFrameColumnDesc",
 			"Game Frame Average Exclusive Duration\n"
 			"Exclusive duration is computed for a single frame as the sum of exclusive duration of all instances of the timer in the respective frame.\n"
 			"The average is computed from these per-frame exclusive durations.\n"
 			"Unit is milliseconds."));
-	Column.SetDescription(ETraceFrameType::TraceFrameType_Rendering,
+	Column.SetDescription(ETimerAggregationMode::RenderingFrame,
 		LOCTEXT("AvgExclusiveTime_RenderingFrameColumnDesc",
 			"Rendering Frame Average Exclusive Duration\n"
 			"Exclusive duration is computed for a single frame as the sum of exclusive duration of all instances of the timer in the respective frame.\n"
@@ -1116,16 +1116,16 @@ TSharedRef<FTableColumn> FTimersViewColumnFactory::CreateMedianExclusiveTimeColu
 	Column.SetShortName(LOCTEXT("MedExclusiveTime_ColumnName", "E.Med"));
 	Column.SetTitleName(LOCTEXT("MedExclusiveTime_ColumnTitle", "Median Exclusive Time (ms)"));
 
-	Column.SetDescription(ETraceFrameType::TraceFrameType_Count,
+	Column.SetDescription(ETimerAggregationMode::Instance,
 		LOCTEXT("MedExclusiveTime_ColumnDesc",
 			"Median exclusive duration of selected timer's instances, in milliseconds"));
-	Column.SetDescription(ETraceFrameType::TraceFrameType_Game,
+	Column.SetDescription(ETimerAggregationMode::GameFrame,
 		LOCTEXT("MedExclusiveTime_GameFrameColumnDesc",
 			"Game Frame Median Exclusive Duration\n"
 			"Exclusive duration is computed for a single frame as the sum of exclusive duration of all instances of the timer in the respective frame.\n"
 			"The median is computed from these per-frame exclusive durations.\n"
 			"Unit is milliseconds."));
-	Column.SetDescription(ETraceFrameType::TraceFrameType_Rendering,
+	Column.SetDescription(ETimerAggregationMode::RenderingFrame,
 		LOCTEXT("MedExclusiveTime_RenderingFrameColumnDesc",
 			"Rendering Frame Median Exclusive Duration\n"
 			"Exclusive duration is computed for a single frame as the sum of exclusive duration of all instances of the timer in the respective frame.\n"
@@ -1175,16 +1175,16 @@ TSharedRef<FTableColumn> FTimersViewColumnFactory::CreateMinExclusiveTimeColumn(
 	Column.SetShortName(LOCTEXT("MinExclusiveTime_ColumnName", "E.Min"));
 	Column.SetTitleName(LOCTEXT("MinExclusiveTime_ColumnTitle", "Min Exclusive Time (ms)"));
 
-	Column.SetDescription(ETraceFrameType::TraceFrameType_Count,
+	Column.SetDescription(ETimerAggregationMode::Instance,
 		LOCTEXT("MinExclusiveTime_ColumnDesc",
 			"Minimum exclusive duration of selected timer's instances, in milliseconds"));
-	Column.SetDescription(ETraceFrameType::TraceFrameType_Game,
+	Column.SetDescription(ETimerAggregationMode::GameFrame,
 		LOCTEXT("MinExclusiveTime_GameFrameColumnDesc",
 			"Game Frame Minimum Exclusive Duration\n"
 			"Exclusive duration is computed for a single frame as the sum of exclusive duration of all instances of the timer in the respective frame.\n"
 			"The minimum is selected from these per-frame exclusive durations.\n"
 			"Unit is milliseconds."));
-	Column.SetDescription(ETraceFrameType::TraceFrameType_Rendering,
+	Column.SetDescription(ETimerAggregationMode::RenderingFrame,
 		LOCTEXT("MinExclusiveTime_RenderingFrameColumnDesc",
 			"Rendering Frame Minimum Exclusive Duration\n"
 			"Exclusive duration is computed for a single frame as the sum of exclusive duration of all instances of the timer in the respective frame.\n"
diff --git a/Engine/Source/Developer/TraceInsights/Private/Insights/TimingProfiler/ViewModels/TimersViewColumnFactory.h b/Engine/Source/Developer/TraceInsights/Private/Insights/TimingProfiler/ViewModels/TimersViewColumnFactory.h
--- a/Engine/Source/Developer/TraceInsights/Private/Insights/TimingProfiler/ViewModels/TimersViewColumnFactory.h
+++ b/Engine/Source/Developer/TraceInsights/Private/Insights/TimingProfiler/ViewModels/TimersViewColumnFactory.h
@@ -11,6 +11,9 @@
 // TraceInsightsCore
 #include "InsightsCore/Table/ViewModels/TableColumn.h"
 
+// TraceInsights
+#include "Insights/TimingProfiler/ViewModels/TimerAggregator.h"
+
 namespace UE::Insights::TimingProfiler
 {
 
@@ -58,32 +61,30 @@ public:
 		: FTableColumn(InId)
 	{}
 
-	FText GetDescription(ETraceFrameType InAggregationMode) const
+	FText GetDescription(ETimerAggregationMode InAggregationMode) const
 	{
 		switch (InAggregationMode)
 		{
-		case TraceFrameType_Game:
+		case ETimerAggregationMode::GameFrame:
 			return GameFrame_Description;
-			break;
-		case TraceFrameType_Rendering:
+		case ETimerAggregationMode::RenderingFrame:
 			return RenderingFrame_Description;
-			break;
 		default:
 			return FTableColumn::GetDescription();
 		}
 	}
 
-	void SetDescription(ETraceFrameType InAggregationMode, FText InDescription)
+	void SetDescription(ETimerAggregationMode InAggregationMode, FText InDescription)
 	{
 		switch (InAggregationMode)
 		{
-		case TraceFrameType_Game:
+		case ETimerAggregationMode::GameFrame:
 			GameFrame_Description = InDescription;
 			break;
-		case TraceFrameType_Rendering:
+		case ETimerAggregationMode::RenderingFrame:
 			RenderingFrame_Description = InDescription;
 			break;
-		case TraceFrameType_Count:
+		case ETimerAggregationMode::Instance:
 			FTableColumn::SetDescription(InDescription);
 			break;
 		default:
diff --git a/Engine/Source/Developer/TraceInsights/Private/Insights/TimingProfiler/Widgets/STimerTreeView.cpp b/Engine/Source/Developer/TraceInsights/Private/Insights/TimingProfiler/Widgets/STimerTreeView.cpp
--- a/Engine/Source/Developer/TraceInsights/Private/Insights/TimingProfiler/Widgets/STimerTreeView.cpp
+++ b/Engine/Source/Developer/TraceInsights/Private/Insights/TimingProfiler/Widgets/STimerTreeView.cpp
@@ -1325,7 +1325,7 @@ void STimerTreeView::ShowColumn(const FName ColumnId)
 	ColumnArgs
 		.ColumnId(Column.GetId())
 		.DefaultLabel(Column.GetShortName())
-		.ToolTip(STimersViewTooltip::GetColumnTooltipForMode(Column, ETraceFrameType::TraceFrameType_Count))
+		.ToolTip(STimersViewTooltip::GetColumnTooltipForMode(Column, ETimerAggregationMode::Instance))
 		.HAlignHeader(Column.GetHorizontalAlignment())
 		.VAlignHeader(VAlign_Center)
 		.HAlignCell(HAlign_Fill)
diff --git a/Engine/Source/Developer/TraceInsights/Private/Insights/TimingProfiler/Widgets/STimersView.cpp b/Engine/Source/Developer/TraceInsights/Private/Insights/TimingProfiler/Widgets/STimersView.cpp
--- a/Engine/Source/Developer/TraceInsights/Private/Insights/TimingProfiler/Widgets/STimersView.cpp
+++ b/Engine/Source/Developer/TraceInsights/Private/Insights/TimingProfiler/Widgets/STimersView.cpp
@@ -468,7 +468,7 @@ void STimersView::Construct(const FArguments& InArgs)
 						SNew(SBox)
 						.MinDesiredWidth(128.0f)
 						[
-							SAssignNew(ModeComboBox, SComboBox<TSharedPtr<ETraceFrameType>>)
+							SAssignNew(ModeComboBox, SComboBox<TSharedPtr<ETimerAggregationMode>>)
 							.ToolTipText(this, &STimersView::Mode_GetSelectedTooltipText)
 							.OptionsSource(&ModeOptionsSource)
 							.OnSelectionChanged(this, &STimersView::Mode_OnSelectionChanged)
@@ -564,7 +564,7 @@ void STimersView::Construct(const FArguments& InArgs)
 	InitializeAndShowHeaderColumns();
 	LoadSettings();
 
-	Aggregator->SetFrameType(ModeFrameType);
+	Aggregator->SetAggregationMode(AggregationMode);
 	SetTimingViewFrameType();
 
 	// Create the search filters: text based, type based etc.
@@ -1579,7 +1579,7 @@ void STimersView::InsightsManager_OnSessionAnalysisCompleted()
 		constexpr double Delta = 0.0; // session padding
 
 		Aggregator->Cancel();
-		Aggregator->SetFrameType(ModeFrameType);
+		Aggregator->SetAggregationMode(AggregationMode);
 		Aggregator->SetTimeInterval(0.0 - Delta, SessionDuration + Delta);
 		Aggregator->Start();
 
@@ -2221,20 +2221,20 @@ FText STimersView::GroupBy_GetSelectedTooltipText() const
 }
 
 ////////////////////////////////////////////////////////////////////////////////////////////////////
-// Frame Type
+// Aggregation Mode
 ////////////////////////////////////////////////////////////////////////////////////////////////////
 
 void STimersView::CreateModeOptionsSources()
 {
-	ModeOptionsSource.Reset(3);
-
-	ModeOptionsSource.Add(MakeShared<ETraceFrameType>(ETraceFrameType::TraceFrameType_Count));
-	ModeOptionsSource.Add(MakeShared<ETraceFrameType>(ETraceFrameType::TraceFrameType_Game));
-	ModeOptionsSource.Add(MakeShared<ETraceFrameType>(ETraceFrameType::TraceFrameType_Rendering));
+	ModeOptionsSource.Reset(uint32(ETimerAggregationMode::Count));
+	for (uint32 Mode = 0; Mode < uint32(ETimerAggregationMode::Count); ++Mode)
+	{
+		ModeOptionsSource.Add(MakeShared<ETimerAggregationMode>(ETimerAggregationMode(Mode)));
+	}
 
-	for (TSharedPtr<ETraceFrameType> Option : ModeOptionsSource)
+	for (TSharedPtr<ETimerAggregationMode> Option : ModeOptionsSource)
 	{
-		if (*Option == ModeFrameType)
+		if (*Option == AggregationMode)
 		{
 			ModeComboBox->SetSelectedItem(Option);
 		}
@@ -2245,52 +2245,52 @@ void STimersView::CreateModeOptionsSources()
 
 ////////////////////////////////////////////////////////////////////////////////////////////////////
 
-void STimersView::Mode_OnSelectionChanged(TSharedPtr<ETraceFrameType> InModeFrameType, ESelectInfo::Type SelectInfo)
+void STimersView::Mode_OnSelectionChanged(TSharedPtr<ETimerAggregationMode> InAggregationMode, ESelectInfo::Type SelectInfo)
 {
-	if (!InModeFrameType.IsValid() || *InModeFrameType == ModeFrameType)
+	if (!InAggregationMode.IsValid() || *InAggregationMode == AggregationMode)
 	{
 		return;
 	}
 
 	SaveVisibleColumnsSettings();
-	ModeFrameType = *InModeFrameType;
+	AggregationMode = *InAggregationMode;
 	LoadVisibleColumnsSettings();
 
 	SetTimingViewFrameType();
 
-	Aggregator->SetFrameType(ModeFrameType);
+	Aggregator->SetAggregationMode(AggregationMode);
 	Aggregator->Cancel();
 	Aggregator->Start();
 }
 
 ////////////////////////////////////////////////////////////////////////////////////////////////////
 
-TSharedRef<SWidget> STimersView::Mode_OnGenerateWidget(TSharedPtr<ETraceFrameType> InModeFrameType) const
+TSharedRef<SWidget> STimersView::Mode_OnGenerateWidget(TSharedPtr<ETimerAggregationMode> InAggregationMode) const
 {
-	check(InModeFrameType.IsValid());
+	check(InAggregationMode.IsValid());
 	return SNew(STextBlock)
-		.Text(Mode_GetText(*InModeFrameType))
-		.ToolTipText(Mode_GetTooltipText(*InModeFrameType));
+		.Text(Mode_GetText(*InAggregationMode))
+		.ToolTipText(Mode_GetTooltipText(*InAggregationMode));
 }
 
 ////////////////////////////////////////////////////////////////////////////////////////////////////
 
 FText STimersView::Mode_GetSelectedText() const
 {
-	return Mode_GetText(ModeFrameType);
+	return Mode_GetText(AggregationMode);
 }
 
 ////////////////////////////////////////////////////////////////////////////////////////////////////
 
-FText STimersView::Mode_GetText(ETraceFrameType InFrameType) const
+FText STimersView::Mode_GetText(ETimerAggregationMode InAggregationMode) const
 {
-	switch (InFrameType)
+	switch (InAggregationMode)
 	{
-	case ETraceFrameType::TraceFrameType_Count:
+	case ETimerAggregationMode::Instance:
 		return LOCTEXT("InstanceMode", "Instance");
-	case ETraceFrameType::TraceFrameType_Game:
+	case ETimerAggregationMode::GameFrame:
 		return LOCTEXT("GameFrameMode", "Game Frame");
-	case ETraceFrameType::TraceFrameType_Rendering:
+	case ETimerAggregationMode::RenderingFrame:
 		return LOCTEXT("RenderingFrameMode", "Rendering Frame");
 	}
 
@@ -2301,20 +2301,20 @@ FText STimersView::Mode_GetText(ETraceFrameType InFrameType) const
 
 FText STimersView::Mode_GetSelectedTooltipText() const
 {
-	return Mode_GetTooltipText(ModeFrameType);
+	return Mode_GetTooltipText(AggregationMode);
 }
 
 ////////////////////////////////////////////////////////////////////////////////////////////////////
 
-FText STimersView::Mode_GetTooltipText(ETraceFrameType InFrameType) const
+FText STimersView::Mode_GetTooltipText(ETimerAggregationMode InAggregationMode) const
 {
-	switch (InFrameType)
+	switch (InAggregationMode)
 	{
-	case ETraceFrameType::TraceFrameType_Count:
+	case ETimerAggregationMode::Instance:
 		return LOCTEXT("ModeInstanceTooltip", "'Instance' Aggregation Mode\nCalculates the stats per event instance.");
-	case ETraceFrameType::TraceFrameType_Game:
+	case ETimerAggregationMode::GameFrame:
 		return LOCTEXT("ModeGameFrameTooltip", "'Game Frame' Aggregation Mode\nCalculates the timer stats per game frame.");
-	case ETraceFrameType::TraceFrameType_Rendering:
+	case ETimerAggregationMode::RenderingFrame:
 		return LOCTEXT("ModeRenderingFrameTooltip", "'Rendering Frame' Aggregation Mode\nCalculates the timer stats per rendering frame.");
 	}
 
@@ -3082,25 +3082,20 @@ TSharedPtr<SFrameTrack> STimersView::GetFrameTrack() const
 
 void STimersView::ToggleTimingViewMainGraphEventSeries(FTimerNodePtr TimerNode) const
 {
-	switch (ModeFrameType)
+	switch (AggregationMode)
 	{
-	case ETraceFrameType::TraceFrameType_Count:
-	{
-		// Instance graph
+	case ETimerAggregationMode::Instance:
 		ToggleTimingViewMainGraphEventInstanceSeries(TimerNode);
 		break;
-	}
-	case ETraceFrameType::TraceFrameType_Game:
-	case ETraceFrameType::TraceFrameType_Rendering:
-	{
-		ToggleTimingViewMainGraphEventFrameStatsSeries(TimerNode, ModeFrameType);
+	case ETimerAggregationMode::GameFrame:
+		ToggleTimingViewMainGraphEventFrameStatsSeries(TimerNode, ETraceFrameType::TraceFrameType_Game);
+		break;
+	case ETimerAggregationMode::RenderingFrame:
+		ToggleTimingViewMainGraphEventFrameStatsSeries(TimerNode, ETraceFrameType::TraceFrameType_Rendering);
 		break;
-	}
 	default:
-	{
 		ensure(0);
 	}
-	}
 }
 
 ////////////////////////////////////////////////////////////////////////////////////////////////////
@@ -4013,23 +4008,17 @@ void STimersView::SaveVisibleColumnsSettings()
 		}
 	}
 
-	switch (ModeFrameType)
+	switch (AggregationMode)
 	{
-		case ETraceFrameType::TraceFrameType_Count:
-		{
+		case ETimerAggregationMode::Instance:
 			Settings.SetTimersViewInstanceVisibleColumns(Columns);
 			break;
-		}
-		case ETraceFrameType::TraceFrameType_Game:
-		{
+		case ETimerAggregationMode::GameFrame:
 			Settings.SetTimersViewGameFrameVisibleColumns(Columns);
 			break;
-		}
-		case ETraceFrameType::TraceFrameType_Rendering:
-		{
+		case ETimerAggregationMode::RenderingFrame:
 			Settings.SetTimersViewRenderingFrameVisibleColumns(Columns);
 			break;
-		}
 	}
 
 	Settings.SaveToConfig();
@@ -4041,7 +4030,7 @@ void STimersView::SaveSettings()
 {
 	FInsightsSettings& Settings = FInsightsManager::Get()->GetSettings();
 
-	Settings.SetTimersViewMode(static_cast<int32>(ModeFrameType));
+	Settings.SetTimersViewAggregationMode(static_cast<int32>(AggregationMode));
 	Settings.SetTimersViewGroupingMode(static_cast<int32>(GroupingMode));
 	Settings.SetTimersViewShowGpuTimers(static_cast<int32>(bTimerTypeIsVisible[static_cast<int32>(ETimerNodeType::GpuScope)]));
 	Settings.SetTimersViewShowVerseSamplingTimers(static_cast<int32>(bTimerTypeIsVisible[static_cast<int32>(ETimerNodeType::VerseSampling)]));
@@ -4059,23 +4048,17 @@ void STimersView::LoadVisibleColumnsSettings()
 	FInsightsSettings& Settings = FInsightsManager::Get()->GetSettings();
 	TArray<FString> Columns;
 
-	switch (ModeFrameType)
+	switch (AggregationMode)
 	{
-		case ETraceFrameType::TraceFrameType_Count:
-		{
+		case ETimerAggregationMode::Instance:
 			Columns = Settings.GetTimersViewInstanceVisibleColumns();
 			break;
-		}
-		case ETraceFrameType::TraceFrameType_Game:
-		{
+		case ETimerAggregationMode::GameFrame:
 			Columns = Settings.GetTimersViewGameFrameVisibleColumns();
 			break;
-		}
-		case ETraceFrameType::TraceFrameType_Rendering:
-		{
+		case ETimerAggregationMode::RenderingFrame:
 			Columns = Settings.GetTimersViewRenderingFrameVisibleColumns();
 			break;
-		}
 	}
 
 	for (const TSharedRef<FTableColumn>& ColumnRef : Table->GetColumns())
@@ -4098,7 +4081,7 @@ void STimersView::LoadSettings()
 {
 	FInsightsSettings& Settings = FInsightsManager::Get()->GetSettings();
 
-	ModeFrameType = static_cast<ETraceFrameType>(Settings.GetTimersViewMode());
+	AggregationMode = static_cast<ETimerAggregationMode>(Settings.GetTimersViewAggregationMode());
 	GroupingMode = static_cast<ETimerGroupingMode>(Settings.GetTimersViewGroupingMode());
 	bTimerTypeIsVisible[static_cast<int32>(ETimerNodeType::GpuScope)] = Settings.GetTimersViewShowGpuTimers();
 	bTimerTypeIsVisible[static_cast<int32>(ETimerNodeType::VerseSampling)] = Settings.GetTimersViewShowVerseSamplingTimers();
@@ -4129,7 +4112,7 @@ void STimersView::OnTimingViewTrackListChanged()
 	if (!Aggregator->IsEmptyTimeInterval())
 	{
 		Aggregator->Cancel();
-		Aggregator->SetFrameType(ModeFrameType);
+		Aggregator->SetAggregationMode(AggregationMode);
 		Aggregator->Start();
 	}
 }
diff --git a/Engine/Source/Developer/TraceInsights/Private/Insights/TimingProfiler/Widgets/STimersView.h b/Engine/Source/Developer/TraceInsights/Private/Insights/TimingProfiler/Widgets/STimersView.h
--- a/Engine/Source/Developer/TraceInsights/Private/Insights/TimingProfiler/Widgets/STimersView.h
+++ b/Engine/Source/Developer/TraceInsights/Private/Insights/TimingProfiler/Widgets/STimersView.h
@@ -17,6 +17,7 @@
 #include "Widgets/Views/STreeView.h"
 
 // TraceInsights
+#include "Insights/TimingProfiler/ViewModels/TimerAggregator.h"
 #include "Insights/TimingProfiler/ViewModels/TimerGroupingAndSorting.h"
 #include "Insights/TimingProfiler/ViewModels/TimerNode.h"
 
@@ -94,7 +95,7 @@ public:
 	FTimerNodePtr GetTimerNode(uint32 TimerId) const;
 	void SelectTimerNode(uint32 TimerId);
 
-	ETraceFrameType GetFrameTypeMode() { return ModeFrameType; }
+	ETimerAggregationMode GetAggregationMode() { return AggregationMode; }
 
 	void OnTimingViewTrackListChanged();
 
@@ -295,12 +296,12 @@ private:
 	// Aggregation Mode
 
 	void CreateModeOptionsSources();
-	void Mode_OnSelectionChanged(TSharedPtr<ETraceFrameType> InModeFrameType, ESelectInfo::Type SelectInfo);
-	TSharedRef<SWidget> Mode_OnGenerateWidget(TSharedPtr<ETraceFrameType> InModeFrameType) const;
+	void Mode_OnSelectionChanged(TSharedPtr<ETimerAggregationMode> InAggregationMode, ESelectInfo::Type SelectInfo);
+	TSharedRef<SWidget> Mode_OnGenerateWidget(TSharedPtr<ETimerAggregationMode> InAggregationMode) const;
 	FText Mode_GetSelectedText() const;
-	FText Mode_GetText(ETraceFrameType InFrameType) const;
+	FText Mode_GetText(ETimerAggregationMode InAggregationMode) const;
 	FText Mode_GetSelectedTooltipText() const;
-	FText Mode_GetTooltipText(ETraceFrameType InFrameType) const;
+	FText Mode_GetTooltipText(ETimerAggregationMode InAggregationMode) const;
 
 	////////////////////////////////////////////////////////////////////////////////////////////////////
 	// Sorting
@@ -502,13 +503,13 @@ private:
 	ETimerGroupingMode GroupingMode = ETimerGroupingMode::ByType;
 
 	//////////////////////////////////////////////////
-	// Frame Type
+	// Aggregation Mode
 
-	TArray<TSharedPtr<ETraceFrameType>> ModeOptionsSource;
+	TArray<TSharedPtr<ETimerAggregationMode>> ModeOptionsSource;
 
-	TSharedPtr<SComboBox<TSharedPtr<ETraceFrameType>>> ModeComboBox;
+	TSharedPtr<SComboBox<TSharedPtr<ETimerAggregationMode>>> ModeComboBox;
 
-	ETraceFrameType ModeFrameType = ETraceFrameType::TraceFrameType_Game;
+	ETimerAggregationMode AggregationMode = ETimerAggregationMode::Instance;
 
 	//////////////////////////////////////////////////
 	// Sorting
diff --git a/Engine/Source/Developer/TraceInsights/Private/Insights/TimingProfiler/Widgets/STimersViewTooltip.cpp b/Engine/Source/Developer/TraceInsights/Private/Insights/TimingProfiler/Widgets/STimersViewTooltip.cpp
--- a/Engine/Source/Developer/TraceInsights/Private/Insights/TimingProfiler/Widgets/STimersViewTooltip.cpp
+++ b/Engine/Source/Developer/TraceInsights/Private/Insights/TimingProfiler/Widgets/STimersViewTooltip.cpp
@@ -73,13 +73,13 @@ TSharedPtr<SToolTip> STimersViewTooltip::GetTableTooltip(const FTable& Table)
 TSharedPtr<SToolTip> STimersViewTooltip::GetColumnTooltip(const FTableColumn& Column)
 {
 	const FTimersTableColumn& TimersColumn = static_cast<const FTimersTableColumn&>(Column);
-	FText InstanceDescription = TimersColumn.GetDescription(ETraceFrameType::TraceFrameType_Count);
-	FText GameFrameDescription = TimersColumn.GetDescription(ETraceFrameType::TraceFrameType_Game);
-	FText RenderingDescription = TimersColumn.GetDescription(ETraceFrameType::TraceFrameType_Rendering);
+	FText InstanceDescription = TimersColumn.GetDescription(ETimerAggregationMode::Instance);
+	FText GameFrameDescription = TimersColumn.GetDescription(ETimerAggregationMode::GameFrame);
+	FText RenderingDescription = TimersColumn.GetDescription(ETimerAggregationMode::RenderingFrame);
 
-	auto GetDescriptionLamda = [InstanceDescription, GameFrameDescription, RenderingDescription]()
+	auto GetDescriptionLambda = [InstanceDescription, GameFrameDescription, RenderingDescription]()
 	{
-		ETraceFrameType FrameType = ETraceFrameType::TraceFrameType_Count;
+		ETimerAggregationMode AggregationMode = ETimerAggregationMode::Instance;
 
 		TSharedPtr<STimingProfilerWindow> Window = FTimingProfilerManager::Get()->GetProfilerWindow();
 		if (Window.IsValid())
@@ -87,17 +87,17 @@ TSharedPtr<SToolTip> STimersViewTooltip::GetColumnTooltip(const FTableColumn& Co
 			TSharedPtr<STimersView> TimersView = Window->GetTimersView();
 			if (TimersView.IsValid())
 			{
-				FrameType = TimersView->GetFrameTypeMode();
+				AggregationMode = TimersView->GetAggregationMode();
 			}
 		}
 
-		switch (FrameType)
+		switch (AggregationMode)
 		{
-		case TraceFrameType_Count:
+		case ETimerAggregationMode::Instance:
 			return InstanceDescription;
-		case TraceFrameType_Game:
+		case ETimerAggregationMode::GameFrame:
 			return GameFrameDescription;
-		case TraceFrameType_Rendering:
+		case ETimerAggregationMode::RenderingFrame:
 			return RenderingDescription;
 		default:
 			ensure(0);
@@ -125,7 +125,7 @@ TSharedPtr<SToolTip> STimersViewTooltip::GetColumnTooltip(const FTableColumn& Co
 			.Padding(2.0f)
 			[
 				SNew(STextBlock)
-				.Text(TAttribute<FText>::CreateLambda(GetDescriptionLamda))
+				.Text(TAttribute<FText>::CreateLambda(GetDescriptionLambda))
 				.TextStyle(FInsightsStyle::Get(), TEXT("TreeTable.Tooltip"))
 			]
 		];
@@ -135,7 +135,7 @@ TSharedPtr<SToolTip> STimersViewTooltip::GetColumnTooltip(const FTableColumn& Co
 
 ////////////////////////////////////////////////////////////////////////////////////////////////////
 
-TSharedPtr<SToolTip> STimersViewTooltip::GetColumnTooltipForMode(const FTableColumn& Column, ETraceFrameType InAggregationMode)
+TSharedPtr<SToolTip> STimersViewTooltip::GetColumnTooltipForMode(const FTableColumn& Column, ETimerAggregationMode InAggregationMode)
 {
 	const FTimersTableColumn& TimersColumn = static_cast<const FTimersTableColumn&>(Column);
 	FText Description = TimersColumn.GetDescription(InAggregationMode);
diff --git a/Engine/Source/Developer/TraceInsights/Private/Insights/TimingProfiler/Widgets/STimersViewTooltip.h b/Engine/Source/Developer/TraceInsights/Private/Insights/TimingProfiler/Widgets/STimersViewTooltip.h
--- a/Engine/Source/Developer/TraceInsights/Private/Insights/TimingProfiler/Widgets/STimersViewTooltip.h
+++ b/Engine/Source/Developer/TraceInsights/Private/Insights/TimingProfiler/Widgets/STimersViewTooltip.h
@@ -7,6 +7,9 @@
 #include "Widgets/IToolTip.h"
 #include "Widgets/SToolTip.h"
 
+// TraceInsights
+#include "Insights/TimingProfiler/ViewModels/TimerAggregator.h"
+
 class SGridPanel;
 
 namespace UE::Insights
@@ -30,7 +33,7 @@ public:
 
 	static TSharedPtr<SToolTip> GetTableTooltip(const FTable& Table);
 	static TSharedPtr<SToolTip> GetColumnTooltip(const FTableColumn& Column);
-	static TSharedPtr<SToolTip> GetColumnTooltipForMode(const FTableColumn& Column, ETraceFrameType InAggregationMode);
+	static TSharedPtr<SToolTip> GetColumnTooltipForMode(const FTableColumn& Column, ETimerAggregationMode InAggregationMode);
 	static TSharedPtr<SToolTip> GetRowTooltip(const TSharedPtr<FTimerNode> TreeNodePtr);
 
 	static bool GetSource(const TSharedPtr<FTimerNode> TreeNodePtr, FText& OutSourcePrefix, FText& OutSourceSuffix);
diff --git a/Engine/Source/Developer/TraceInsights/Private/Insights/Widgets/STimingView.cpp b/Engine/Source/Developer/TraceInsights/Private/Insights/Widgets/STimingView.cpp
--- a/Engine/Source/Developer/TraceInsights/Private/Insights/Widgets/STimingView.cpp
+++ b/Engine/Source/Developer/TraceInsights/Private/Insights/Widgets/STimingView.cpp
@@ -6344,11 +6344,17 @@ ETraceFrameType STimingView::GetFrameTypeToSnapTo()
 	TSharedPtr<STimersView> TimersView = GetTimersView();
 	if (TimersView.IsValid())
 	{
-		return TimersView->GetFrameTypeMode();
+		ETimerAggregationMode AggregationMode = TimersView->GetAggregationMode();
+		switch (AggregationMode)
+		{
+		case ETimerAggregationMode::GameFrame:
+			return ETraceFrameType::TraceFrameType_Game;
+		case ETimerAggregationMode::RenderingFrame:
+			return ETraceFrameType::TraceFrameType_Rendering;
+		}
 	}
 
-	// TraceFrameType_Count is the Instance mode.
-	return ETraceFrameType::TraceFrameType_Count;
+	return ETraceFrameType::TraceFrameType_Count; // no frame snap
 }
 
 ////////////////////////////////////////////////////////////////////////////////////////////////////
