[Insights]
- Removed deprecated FNames from FInsightsManagerTabs.

#rnx

#Synced-CL 51092797

diff --git a/Engine/Source/Developer/TraceInsights/Private/Insights/InsightsManager.cpp b/Engine/Source/Developer/TraceInsights/Private/Insights/InsightsManager.cpp
--- a/Engine/Source/Developer/TraceInsights/Private/Insights/InsightsManager.cpp
+++ b/Engine/Source/Developer/TraceInsights/Private/Insights/InsightsManager.cpp
@@ -54,11 +54,6 @@
 
 #define LOCTEXT_NAMESPACE "InsightsManager"
 
-const FName FInsightsManagerTabs::StartPageTabId(TEXT("TraceStore")); // DEPRECATED
-const FName FInsightsManagerTabs::TraceStoreTabId(TEXT("TraceStore")); // DEPRECATED
-const FName FInsightsManagerTabs::ConnectionTabId(TEXT("Connection")); // DEPRECATED
-const FName FInsightsManagerTabs::LauncherTabId(TEXT("Launcher")); // DEPRECATED
-
 const FName FInsightsManagerTabs::SessionInfoTabId(TEXT("SessionInfo"));
 const FName FInsightsManagerTabs::TimingProfilerTabId(TEXT("TimingProfiler"));
 const FName FInsightsManagerTabs::LoadingProfilerTabId(TEXT("LoadingProfiler"));
diff --git a/Engine/Source/Developer/TraceInsights/Public/Insights/IUnrealInsightsModule.h b/Engine/Source/Developer/TraceInsights/Public/Insights/IUnrealInsightsModule.h
--- a/Engine/Source/Developer/TraceInsights/Public/Insights/IUnrealInsightsModule.h
+++ b/Engine/Source/Developer/TraceInsights/Public/Insights/IUnrealInsightsModule.h
@@ -11,9 +11,11 @@
 #include "Internationalization/Text.h"
 #include "Modules/ModuleInterface.h"
 #include "Templates/SharedPointer.h"
-#include "Trace/DataStream.h"
 #include "UObject/NameTypes.h"
 
+// TraceAnalysis
+#include "Trace/DataStream.h"
+
 #define UE_API TRACEINSIGHTS_API
 
 class FExtender;
@@ -28,11 +30,6 @@ namespace Insights { class IInsightsManager; }
 /** Major tab IDs for Insights tools */
 struct FInsightsManagerTabs
 {
-	static UE_API const FName StartPageTabId; // DEPRECATED
-	static UE_API const FName TraceStoreTabId; // DEPRECATED: see FInsightsFrontendTabs::TraceStoreTabId
-	static UE_API const FName ConnectionTabId; // DEPRECATED: see FInsightsFrontendTabs::ConnectionTabId
-	static UE_API const FName LauncherTabId; // DEPRECATED
-
 	static UE_API const FName SessionInfoTabId;
 	static UE_API const FName TimingProfilerTabId;
 	static UE_API const FName LoadingProfilerTabId;
