[Insights] TraceInsightsCore: Deprecated the FFilterConfiguratorNode constructor with FName param.

#Synced-CL 51157224

diff --git a/Engine/Source/Developer/TraceInsights/Private/Insights/Widgets/STimingView.cpp b/Engine/Source/Developer/TraceInsights/Private/Insights/Widgets/STimingView.cpp
--- a/Engine/Source/Developer/TraceInsights/Private/Insights/Widgets/STimingView.cpp
+++ b/Engine/Source/Developer/TraceInsights/Private/Insights/Widgets/STimingView.cpp
@@ -416,17 +416,17 @@ TSharedPtr<FFilterConfigurator> STimingView::MakeCustomFilterConfigurator(uint32
 		nullptr,
 		FFilterService::Get()->GetIntegerOperators()));
 
-	FFilterConfiguratorNodePtr StartTimeNode = MakeShared<FFilterConfiguratorNode>(TEXT(""), false);
+	FFilterConfiguratorNodePtr StartTimeNode = MakeShared<FFilterConfiguratorNode>(false);
 	StartTimeNode->SetSelectedFilter(CustomFilterConfigurator->GetRootNode()->GetAvailableFilters()->GetData()[0]);
 	StartTimeNode->SetTextBoxValue(FString::SanitizeFloat(IntervalStart));
 	StartTimeNode->SetSelectedFilterOperator(MakeShared<FFilterOperator<double>>(EFilterOperator::Gte, TEXT(""), std::greater_equal<>()));
 
-	FFilterConfiguratorNodePtr EndTimeNode = MakeShared<FFilterConfiguratorNode>(TEXT(""), false);
+	FFilterConfiguratorNodePtr EndTimeNode = MakeShared<FFilterConfiguratorNode>(false);
 	EndTimeNode->SetSelectedFilter(CustomFilterConfigurator->GetRootNode()->GetAvailableFilters()->GetData()[1]);
 	EndTimeNode->SetTextBoxValue(FString::SanitizeFloat(IntervalEnd));
 
 	EndTimeNode->SetSelectedFilterOperator(MakeShared<FFilterOperator<double>>(EFilterOperator::Lte, TEXT(""), std::less_equal<>()));
-	FFilterConfiguratorNodePtr TimerIdNode = MakeShared<FFilterConfiguratorNode>(TEXT(""), false);
+	FFilterConfiguratorNodePtr TimerIdNode = MakeShared<FFilterConfiguratorNode>(false);
 	TimerIdNode->SetSelectedFilter(CustomFilterConfigurator->GetRootNode()->GetAvailableFilters()->GetData()[2]);
 	TimerIdNode->SetTextBoxValue(FString::FromInt(TimerId));
 
diff --git a/Engine/Source/Developer/TraceInsightsCore/Private/InsightsCore/Filter/ViewModels/FilterConfigurator.cpp b/Engine/Source/Developer/TraceInsightsCore/Private/InsightsCore/Filter/ViewModels/FilterConfigurator.cpp
--- a/Engine/Source/Developer/TraceInsightsCore/Private/InsightsCore/Filter/ViewModels/FilterConfigurator.cpp
+++ b/Engine/Source/Developer/TraceInsightsCore/Private/InsightsCore/Filter/ViewModels/FilterConfigurator.cpp
@@ -11,7 +11,7 @@ namespace UE::Insights
 
 FFilterConfigurator::FFilterConfigurator()
 {
-	RootNode = MakeShared<FFilterConfiguratorNode>(TEXT(""), true);
+	RootNode = MakeShared<FFilterConfiguratorNode>(true);
 	AvailableFilters = MakeShared<TArray<TSharedPtr<FFilter>>>();
 	RootNode->SetAvailableFilters(AvailableFilters);
 }
diff --git a/Engine/Source/Developer/TraceInsightsCore/Private/InsightsCore/Filter/ViewModels/FilterConfiguratorNode.cpp b/Engine/Source/Developer/TraceInsightsCore/Private/InsightsCore/Filter/ViewModels/FilterConfiguratorNode.cpp
--- a/Engine/Source/Developer/TraceInsightsCore/Private/InsightsCore/Filter/ViewModels/FilterConfiguratorNode.cpp
+++ b/Engine/Source/Developer/TraceInsightsCore/Private/InsightsCore/Filter/ViewModels/FilterConfiguratorNode.cpp
@@ -15,8 +15,8 @@ INSIGHTS_IMPLEMENT_RTTI(FFilterConfiguratorNode)
 
 ////////////////////////////////////////////////////////////////////////////////////////////////////
 
-FFilterConfiguratorNode::FFilterConfiguratorNode(const FName InName, bool bInIsGroup)
-	: FBaseTreeNode(InName, bInIsGroup)
+FFilterConfiguratorNode::FFilterConfiguratorNode(bool bInIsGroup)
+	: FBaseTreeNode(bInIsGroup)
 {
 	if (bInIsGroup)
 	{
@@ -30,7 +30,7 @@ FFilterConfiguratorNode::FFilterConfiguratorNode(const FName InName, bool bInIsG
 
 TSharedRef<FFilterConfiguratorNode> FFilterConfiguratorNode::DeepCopy(const FFilterConfiguratorNode& InSourceNode)
 {
-	TSharedRef<FFilterConfiguratorNode> NodeCopyPtr = MakeShared<FFilterConfiguratorNode>(InSourceNode.GetName(), InSourceNode.IsGroup());
+	TSharedRef<FFilterConfiguratorNode> NodeCopyPtr = MakeShared<FFilterConfiguratorNode>(InSourceNode.IsGroup());
 	FFilterConfiguratorNode& NodeCopy = *NodeCopyPtr;
 
 	NodeCopy.AvailableFilters = InSourceNode.AvailableFilters;
diff --git a/Engine/Source/Developer/TraceInsightsCore/Private/InsightsCore/Filter/Widgets/SFilterConfiguratorRow.cpp b/Engine/Source/Developer/TraceInsightsCore/Private/InsightsCore/Filter/Widgets/SFilterConfiguratorRow.cpp
--- a/Engine/Source/Developer/TraceInsightsCore/Private/InsightsCore/Filter/Widgets/SFilterConfiguratorRow.cpp
+++ b/Engine/Source/Developer/TraceInsightsCore/Private/InsightsCore/Filter/Widgets/SFilterConfiguratorRow.cpp
@@ -412,7 +412,7 @@ FText SFilterConfiguratorRow::FilterGroupOperators_GetSelectionText() const
 
 FReply SFilterConfiguratorRow::AddFilter_OnClicked()
 {
-	FFilterConfiguratorNodePtr ChildNode = MakeShared<FFilterConfiguratorNode>(TEXT(""), false);
+	FFilterConfiguratorNodePtr ChildNode = MakeShared<FFilterConfiguratorNode>(false);
 	ChildNode->SetAvailableFilters(FilterConfiguratorNodePtr->GetAvailableFilters());
 
 	FilterConfiguratorNodePtr->AddChildAndSetParent(ChildNode);
@@ -437,7 +437,7 @@ FReply SFilterConfiguratorRow::DeleteFilter_OnClicked()
 
 FReply SFilterConfiguratorRow::AddGroup_OnClicked()
 {
-	FFilterConfiguratorNodePtr ChildNode = MakeShared<FFilterConfiguratorNode>(TEXT(""), true);
+	FFilterConfiguratorNodePtr ChildNode = MakeShared<FFilterConfiguratorNode>(true);
 	ChildNode->SetAvailableFilters(FilterConfiguratorNodePtr->GetAvailableFilters());
 	ChildNode->SetExpansion(true);
 
diff --git a/Engine/Source/Developer/TraceInsightsCore/Public/InsightsCore/Filter/ViewModels/FilterConfiguratorNode.h b/Engine/Source/Developer/TraceInsightsCore/Public/InsightsCore/Filter/ViewModels/FilterConfiguratorNode.h
--- a/Engine/Source/Developer/TraceInsightsCore/Public/InsightsCore/Filter/ViewModels/FilterConfiguratorNode.h
+++ b/Engine/Source/Developer/TraceInsightsCore/Public/InsightsCore/Filter/ViewModels/FilterConfiguratorNode.h
@@ -36,7 +36,10 @@ class FFilterConfiguratorNode : public FBaseTreeNode
 
 public:
 	/** Initialization constructor for the filter configurator node. */
-	UE_API FFilterConfiguratorNode(const FName InName, bool bInIsGroup);
+	UE_API FFilterConfiguratorNode(bool bInIsGroup);
+
+	UE_DEPRECATED(5.8, "Use constructor without Name")
+	UE_API FFilterConfiguratorNode(FName InName, bool bInIsGroup) : FFilterConfiguratorNode(bInIsGroup) {}
 
 	static UE_API TSharedRef<FFilterConfiguratorNode> DeepCopy(const FFilterConfiguratorNode& Node);
 
