Unreal Insights : sort the analysis modules by name when displaying them in the session tab

#Synced-CL 50949381

diff --git a/Engine/Source/Developer/TraceInsights/Private/Insights/Widgets/SSessionInfoWindow.cpp b/Engine/Source/Developer/TraceInsights/Private/Insights/Widgets/SSessionInfoWindow.cpp
--- a/Engine/Source/Developer/TraceInsights/Private/Insights/Widgets/SSessionInfoWindow.cpp
+++ b/Engine/Source/Developer/TraceInsights/Private/Insights/Widgets/SSessionInfoWindow.cpp
@@ -684,6 +684,11 @@ FText SSessionInfoWindow::GetModulesText() const
 		ModuleService->GetAvailableModulesEx(Modules);
 	}
 
+	Modules.Sort([](const TraceServices::FModuleInfoEx& LHS, const TraceServices::FModuleInfoEx& RHS)
+		{
+			return FCString::Strcmp(LHS.Info.DisplayName, RHS.Info.DisplayName) < 0;
+		});
+
 	bool bIsFirst = true;
 	for (const TraceServices::FModuleInfoEx& Module : Modules)
 	{
