[Insights]
- Fixed includes and forward declarations for all profilers (IInsightsComponent implementations).
- Added FTimingProfilerManager::ActivateMajorTab() and ActivateView() -- moved functionality from FInsightsManager.

#rnx

#Synced-CL 51163196

#Synced-CL 51163196

diff --git a/Engine/Source/Developer/TraceInsights/Private/Insights/ContextSwitches/ContextSwitchesProfilerManager.h b/Engine/Source/Developer/TraceInsights/Private/Insights/ContextSwitches/ContextSwitchesProfilerManager.h
--- a/Engine/Source/Developer/TraceInsights/Private/Insights/ContextSwitches/ContextSwitchesProfilerManager.h
+++ b/Engine/Source/Developer/TraceInsights/Private/Insights/ContextSwitches/ContextSwitchesProfilerManager.h
@@ -5,13 +5,15 @@
 #include "CoreTypes.h"
 
 #include "Containers/Ticker.h"
+#include "Framework/Commands/UICommandList.h"
+#include "Templates/SharedPointer.h"
+
+// TraceInsightsCore
+#include "InsightsCore/Common/AvailabilityCheck.h"
 
 // TraceInsights
-#include "Insights/InsightsManager.h"
 #include "Insights/IUnrealInsightsModule.h"
 
-////////////////////////////////////////////////////////////////////////////////////////////////////
-
 class FThreadTrackEvent;
 class FUICommandList;
 
@@ -78,6 +80,4 @@ private:
 	TSharedPtr<FContextSwitchesSharedState> ContextSwitchesSharedState;
 };
 
-////////////////////////////////////////////////////////////////////////////////////////////////////
-
 } // namespace UE::Insights::ContextSwitches
diff --git a/Engine/Source/Developer/TraceInsights/Private/Insights/CookProfiler/CookProfilerManager.cpp b/Engine/Source/Developer/TraceInsights/Private/Insights/CookProfiler/CookProfilerManager.cpp
--- a/Engine/Source/Developer/TraceInsights/Private/Insights/CookProfiler/CookProfilerManager.cpp
+++ b/Engine/Source/Developer/TraceInsights/Private/Insights/CookProfiler/CookProfilerManager.cpp
@@ -12,6 +12,7 @@
 #include "TraceServices/Model/CookProfilerProvider.h"
 
 // TraceInsights
+#include "Insights/InsightsManager.h"
 #include "Insights/InsightsStyle.h"
 #include "Insights/CookProfiler/ViewModels/PackageTable.h"
 #include "Insights/CookProfiler/Widgets/SPackageTableTreeView.h"
diff --git a/Engine/Source/Developer/TraceInsights/Private/Insights/CookProfiler/CookProfilerManager.h b/Engine/Source/Developer/TraceInsights/Private/Insights/CookProfiler/CookProfilerManager.h
--- a/Engine/Source/Developer/TraceInsights/Private/Insights/CookProfiler/CookProfilerManager.h
+++ b/Engine/Source/Developer/TraceInsights/Private/Insights/CookProfiler/CookProfilerManager.h
@@ -8,8 +8,10 @@
 #include "Framework/Commands/UICommandList.h"
 #include "Templates/SharedPointer.h"
 
+// TraceInsightsCore
+#include "InsightsCore/Common/AvailabilityCheck.h"
+
 // TraceInsights
-#include "Insights/InsightsManager.h"
 #include "Insights/IUnrealInsightsModule.h"
 
 namespace UE::Insights::CookProfiler
@@ -17,6 +19,8 @@ namespace UE::Insights::CookProfiler
 
 class SPackageTableTreeView;
 
+////////////////////////////////////////////////////////////////////////////////////////////////////
+
 struct FCookProfilerTabs
 {
 	// Tab identifiers
@@ -95,4 +99,3 @@ private:
 };
 
 } // namespace UE::Insights::CookProfiler
-
diff --git a/Engine/Source/Developer/TraceInsights/Private/Insights/InsightsManager.cpp b/Engine/Source/Developer/TraceInsights/Private/Insights/InsightsManager.cpp
--- a/Engine/Source/Developer/TraceInsights/Private/Insights/InsightsManager.cpp
+++ b/Engine/Source/Developer/TraceInsights/Private/Insights/InsightsManager.cpp
@@ -16,23 +16,19 @@
 #include "Misc/FileHelper.h"
 #include "Misc/Paths.h"
 #include "Modules/ModuleManager.h"
-#include "Templates/UniquePtr.h"
-#include "Trace/DataStream.h"
-#include "Trace/SecureDataStream.h"
-#include "WorkspaceMenuStructure.h"
-#include "WorkspaceMenuStructureModule.h"
 
 // TraceTools
 #include "TraceTools/Interfaces/ITraceToolsModule.h"
 
 // TraceAnalysis
+#include "Trace/DataStream.h"
+#include "Trace/SecureDataStream.h"
 #include "Trace/StoreClient.h"
 
 // TraceServices
 #include "TraceServices/AnalysisService.h"
 #include "TraceServices/Model/Diagnostics.h"
 #include "TraceServices/Model/Modules.h"
-#include "TraceServices/Model/NetProfiler.h"
 
 // TraceInsightsCore
 #include "InsightsCore/Common/MiscUtils.h"
@@ -42,12 +38,9 @@
 #include "Insights/Common/InsightsMenuBuilder.h"
 #include "Insights/InsightsStyle.h"
 #include "Insights/IUnrealInsightsModule.h"
-#include "Insights/LoadingProfiler/LoadingProfilerManager.h"
 #include "Insights/Log.h"
-#include "Insights/NetworkingProfiler/NetworkingProfilerManager.h"
 #include "Insights/Tests/InsightsTestRunner.h"
 #include "Insights/TimingProfiler/TimingProfilerManager.h"
-#include "Insights/TimingProfiler/Widgets/STimingProfilerWindow.h"
 #include "Insights/Widgets/SSessionInfoWindow.h"
 
 ////////////////////////////////////////////////////////////////////////////////////////////////////
@@ -578,25 +571,7 @@ void FInsightsManager::SpawnAndActivateTabs()
 void FInsightsManager::ActivateTimingInsightsTab()
 {
 	// Ensure Timing Insights / Timing View is the active tab / view.
-	if (TSharedPtr<SDockTab> TimingInsightsTab = FGlobalTabmanager::Get()->FindExistingLiveTab(FInsightsManagerTabs::TimingProfilerTabId))
-	{
-		TimingInsightsTab->ActivateInParent(ETabActivationCause::SetDirectly);
-
-		using namespace UE::Insights::TimingProfiler;
-
-		//TODO: FTimingProfilerManager::Get()->ActivateWindow();
-		TSharedPtr<class STimingProfilerWindow> Wnd = FTimingProfilerManager::Get()->GetProfilerWindow();
-		if (Wnd)
-		{
-			TSharedPtr<FTabManager> TabManager = Wnd->GetTabManager();
-
-			if (TSharedPtr<SDockTab> TimingViewTab = TabManager->FindExistingLiveTab(FTimingProfilerTabs::TimingViewID))
-			{
-				TimingViewTab->ActivateInParent(ETabActivationCause::SetDirectly);
-				FSlateApplication::Get().SetKeyboardFocus(TimingViewTab->GetContent());
-			}
-		}
-	}
+	TimingProfiler::FTimingProfilerManager::Get()->ActivateMajorTab();
 }
 
 ////////////////////////////////////////////////////////////////////////////////////////////////////
diff --git a/Engine/Source/Developer/TraceInsights/Private/Insights/InsightsManager.h b/Engine/Source/Developer/TraceInsights/Private/Insights/InsightsManager.h
--- a/Engine/Source/Developer/TraceInsights/Private/Insights/InsightsManager.h
+++ b/Engine/Source/Developer/TraceInsights/Private/Insights/InsightsManager.h
@@ -2,18 +2,22 @@
 
 #pragma once
 
-#include "Containers/StringFwd.h"
 #include "CoreTypes.h"
 
 #include "Containers/Ticker.h"
+#include "Containers/UnrealString.h"
 #include "Framework/Commands/UICommandList.h"
 #include "Input/DragAndDrop.h"
+#include "Templates/SharedPointer.h"
+#include "Templates/UniquePtr.h"
 
 // TraceAnalysis
 #include "Trace/StoreConnection.h"
 
+// TraceServices
+#include "TraceServices/Model/AnalysisSession.h"
+
 // TraceInsightsCore
-#include "InsightsCore/Common/AvailabilityCheck.h"
 #include "InsightsCore/Common/Stopwatch.h"
 
 // TraceInsights
@@ -21,9 +25,6 @@
 #include "Insights/IInsightsManager.h"
 #include "Insights/InsightsCommands.h"
 #include "Insights/InsightsSettings.h"
-#include "Insights/IUnrealInsightsModule.h"
-
-////////////////////////////////////////////////////////////////////////////////////////////////////
 
 namespace UE::Trace
 {
@@ -36,6 +37,7 @@ namespace TraceServices
 	class IModuleService;
 }
 
+class IUnrealInsightsModule;
 class FInsightsTestRunner;
 
 namespace UE::Insights
@@ -44,6 +46,7 @@ namespace UE::Insights
 class SSessionInfoWindow;
 class FInsightsMenuBuilder;
 
+////////////////////////////////////////////////////////////////////////////////////////////////////
 /**
  * Struct that holds data about in progress async operations
  */
@@ -59,6 +62,8 @@ struct FAsyncTaskData
 	}
 };
 
+////////////////////////////////////////////////////////////////////////////////////////////////////
+
 enum class ETraceStreamType
 {
 	Unknown,
@@ -346,7 +351,7 @@ private:
 
 	/** Extract messages from the session */
 	void PollAnalysisInfo();
-	
+
 	void OnSessionChanged();
 	void OnSessionAnalysisCompleted();
 	void OnSessionSymbolResolverCompleted();
diff --git a/Engine/Source/Developer/TraceInsights/Private/Insights/MemoryProfiler/MemoryProfilerManager.h b/Engine/Source/Developer/TraceInsights/Private/Insights/MemoryProfiler/MemoryProfilerManager.h
--- a/Engine/Source/Developer/TraceInsights/Private/Insights/MemoryProfiler/MemoryProfilerManager.h
+++ b/Engine/Source/Developer/TraceInsights/Private/Insights/MemoryProfiler/MemoryProfilerManager.h
@@ -6,10 +6,13 @@
 
 #include "Containers/Ticker.h"
 #include "Framework/Commands/UICommandList.h"
+#include "Templates/SharedPointer.h"
 #include "Logging/LogMacros.h"
 
+// TraceInsightsCore
+#include "InsightsCore/Common/AvailabilityCheck.h"
+
 // TraceInsights
-#include "Insights/InsightsManager.h"
 #include "Insights/IUnrealInsightsModule.h"
 #include "Insights/MemoryProfiler/MemoryProfilerCommands.h"
 
diff --git a/Engine/Source/Developer/TraceInsights/Private/Insights/MemoryProfiler/Widgets/SMemAllocTableTreeView.cpp b/Engine/Source/Developer/TraceInsights/Private/Insights/MemoryProfiler/Widgets/SMemAllocTableTreeView.cpp
--- a/Engine/Source/Developer/TraceInsights/Private/Insights/MemoryProfiler/Widgets/SMemAllocTableTreeView.cpp
+++ b/Engine/Source/Developer/TraceInsights/Private/Insights/MemoryProfiler/Widgets/SMemAllocTableTreeView.cpp
@@ -38,6 +38,7 @@
 
 // TraceInsights
 #include "Insights/Common/Utf8FileWriter.h"
+#include "Insights/InsightsManager.h"
 #include "Insights/InsightsStyle.h"
 #include "Insights/MemoryProfiler/Common/SymbolSearchPathsHelper.h"
 #include "Insights/MemoryProfiler/MemoryProfilerManager.h"
diff --git a/Engine/Source/Developer/TraceInsights/Private/Insights/MemoryProfiler/Widgets/SMemInvestigationView.cpp b/Engine/Source/Developer/TraceInsights/Private/Insights/MemoryProfiler/Widgets/SMemInvestigationView.cpp
--- a/Engine/Source/Developer/TraceInsights/Private/Insights/MemoryProfiler/Widgets/SMemInvestigationView.cpp
+++ b/Engine/Source/Developer/TraceInsights/Private/Insights/MemoryProfiler/Widgets/SMemInvestigationView.cpp
@@ -14,6 +14,7 @@
 #include "InsightsCore/Common/TimeUtils.h"
 
 // TraceInsights
+#include "Insights/InsightsManager.h"
 #include "Insights/InsightsStyle.h"
 #include "Insights/MemoryProfiler/Common/SymbolSearchPathsHelper.h"
 #include "Insights/MemoryProfiler/MemoryProfilerManager.h"
diff --git a/Engine/Source/Developer/TraceInsights/Private/Insights/NetworkingProfiler/NetworkingProfilerManager.h b/Engine/Source/Developer/TraceInsights/Private/Insights/NetworkingProfiler/NetworkingProfilerManager.h
--- a/Engine/Source/Developer/TraceInsights/Private/Insights/NetworkingProfiler/NetworkingProfilerManager.h
+++ b/Engine/Source/Developer/TraceInsights/Private/Insights/NetworkingProfiler/NetworkingProfilerManager.h
@@ -6,10 +6,13 @@
 
 #include "Containers/Ticker.h"
 #include "Framework/Commands/UICommandList.h"
+#include "Templates/SharedPointer.h"
 #include "Logging/LogMacros.h"
 
+// TraceInsightsCore
+#include "InsightsCore/Common/AvailabilityCheck.h"
+
 // TraceInsights
-#include "Insights/InsightsManager.h"
 #include "Insights/IUnrealInsightsModule.h"
 #include "Insights/NetworkingProfiler/NetworkingProfilerCommands.h"
 
diff --git a/Engine/Source/Developer/TraceInsights/Private/Insights/NetworkingProfiler/Widgets/SNetStatsCountersView.cpp b/Engine/Source/Developer/TraceInsights/Private/Insights/NetworkingProfiler/Widgets/SNetStatsCountersView.cpp
--- a/Engine/Source/Developer/TraceInsights/Private/Insights/NetworkingProfiler/Widgets/SNetStatsCountersView.cpp
+++ b/Engine/Source/Developer/TraceInsights/Private/Insights/NetworkingProfiler/Widgets/SNetStatsCountersView.cpp
@@ -22,6 +22,7 @@
 #include "InsightsCore/Table/ViewModels/TableColumn.h"
 
 // TraceInsights
+#include "Insights/InsightsManager.h"
 #include "Insights/InsightsStyle.h"
 #include "Insights/NetworkingProfiler/NetworkingProfilerManager.h"
 #include "Insights/NetworkingProfiler/ViewModels/NetStatsCounterNodeHelper.h"
diff --git a/Engine/Source/Developer/TraceInsights/Private/Insights/NetworkingProfiler/Widgets/SNetStatsView.cpp b/Engine/Source/Developer/TraceInsights/Private/Insights/NetworkingProfiler/Widgets/SNetStatsView.cpp
--- a/Engine/Source/Developer/TraceInsights/Private/Insights/NetworkingProfiler/Widgets/SNetStatsView.cpp
+++ b/Engine/Source/Developer/TraceInsights/Private/Insights/NetworkingProfiler/Widgets/SNetStatsView.cpp
@@ -18,10 +18,12 @@
 #include "TraceServices/Model/NetProfiler.h"
 
 // TraceInsightsCore
+#include "InsightsCore/Common/Stopwatch.h"
 #include "InsightsCore/Table/ViewModels/Table.h"
 #include "InsightsCore/Table/ViewModels/TableColumn.h"
 
 // TraceInsights
+#include "Insights/InsightsManager.h"
 #include "Insights/InsightsStyle.h"
 #include "Insights/NetworkingProfiler/NetworkingProfilerManager.h"
 #include "Insights/NetworkingProfiler/ViewModels/NetStatsViewColumnFactory.h"
diff --git a/Engine/Source/Developer/TraceInsights/Private/Insights/NetworkingProfiler/Widgets/SNetworkingProfilerWindow.cpp b/Engine/Source/Developer/TraceInsights/Private/Insights/NetworkingProfiler/Widgets/SNetworkingProfilerWindow.cpp
--- a/Engine/Source/Developer/TraceInsights/Private/Insights/NetworkingProfiler/Widgets/SNetworkingProfilerWindow.cpp
+++ b/Engine/Source/Developer/TraceInsights/Private/Insights/NetworkingProfiler/Widgets/SNetworkingProfilerWindow.cpp
@@ -10,6 +10,7 @@
 #include "Widgets/Docking/SDockTab.h"
 
 // TraceInsights
+#include "Insights/InsightsManager.h"
 #include "Insights/InsightsStyle.h"
 #include "Insights/NetworkingProfiler/NetworkingProfilerManager.h"
 #include "Insights/NetworkingProfiler/Widgets/SNetStatsView.h"
@@ -55,7 +56,7 @@ SNetworkingProfilerWindow::~SNetworkingProfilerWindow()
 	CloseAllOpenTabs();
 
 	check(NetStatsView == nullptr);
-	check(NetStatsCountersView = nullptr);
+	check(NetStatsCountersView == nullptr);
 	check(PacketContentView == nullptr);
 	check(PacketView == nullptr);
 }
@@ -91,6 +92,7 @@ void SNetworkingProfilerWindow::Reset()
 	{
 		NetStatsCountersView->Reset();
 	}
+
 	AvailableGameInstances.Reset();
 	SelectedGameInstance.Reset();
 	AvailableConnections.Reset();
diff --git a/Engine/Source/Developer/TraceInsights/Private/Insights/TaskGraphProfiler/TaskGraphProfilerManager.h b/Engine/Source/Developer/TraceInsights/Private/Insights/TaskGraphProfiler/TaskGraphProfilerManager.h
--- a/Engine/Source/Developer/TraceInsights/Private/Insights/TaskGraphProfiler/TaskGraphProfilerManager.h
+++ b/Engine/Source/Developer/TraceInsights/Private/Insights/TaskGraphProfiler/TaskGraphProfilerManager.h
@@ -9,8 +9,10 @@
 #include "Framework/Commands/UICommandList.h"
 #include "Templates/SharedPointer.h"
 
+// TraceInsightsCore
+#include "InsightsCore/Common/AvailabilityCheck.h"
+
 // TraceInsights
-#include "Insights/InsightsManager.h"
 #include "Insights/IUnrealInsightsModule.h"
 
 namespace TraceServices
@@ -58,7 +60,7 @@ enum class ETaskEventType : uint32
 	NestedCompleted,
 	SubsequentStarted,
 	CriticalPath,
-	
+
 	NumTaskEventTypes,
 };
 
@@ -137,7 +139,7 @@ public:
 
 	bool GetShowNestedTasks() const { return bShowNestedTasks; }
 	void SetShowNestedTasks(bool bInValue) { bShowNestedTasks = bInValue; }
-	
+
 	bool GetShowAnyRelations() const { return GetShowCriticalPath() || GetShowTransitions() || GetShowConnections() || GetShowPrerequisites() || GetShowSubsequents() || GetShowParentTasks() || GetShowNestedTasks(); }
 
 	int32 GetDepthOfTaskExecution(double TaskStartedTime, double TaskFinishedTime, uint32 ThreadId);
@@ -195,7 +197,7 @@ private:
 	bool bShowSubsequents = false;
 	bool bShowParentTasks = false;
 	bool bShowNestedTasks = false;
-	
+
 	FDelegateHandle OnWindowClosedEventHandle;
 
 	TSet<FString> HiddenTrackNames;
diff --git a/Engine/Source/Developer/TraceInsights/Private/Insights/TaskGraphProfiler/ViewModels/TaskTimingTrack.cpp b/Engine/Source/Developer/TraceInsights/Private/Insights/TaskGraphProfiler/ViewModels/TaskTimingTrack.cpp
--- a/Engine/Source/Developer/TraceInsights/Private/Insights/TaskGraphProfiler/ViewModels/TaskTimingTrack.cpp
+++ b/Engine/Source/Developer/TraceInsights/Private/Insights/TaskGraphProfiler/ViewModels/TaskTimingTrack.cpp
@@ -14,6 +14,7 @@
 #include "InsightsCore/Common/TimeUtils.h"
 
 // TraceInsights
+#include "Insights/InsightsManager.h"
 #include "Insights/InsightsStyle.h"
 #include "Insights/ITimingViewSession.h"
 #include "Insights/TaskGraphProfiler/TaskGraphProfilerManager.h"
diff --git a/Engine/Source/Developer/TraceInsights/Private/Insights/TimingProfiler/TimingProfilerManager.cpp b/Engine/Source/Developer/TraceInsights/Private/Insights/TimingProfiler/TimingProfilerManager.cpp
--- a/Engine/Source/Developer/TraceInsights/Private/Insights/TimingProfiler/TimingProfilerManager.cpp
+++ b/Engine/Source/Developer/TraceInsights/Private/Insights/TimingProfiler/TimingProfilerManager.cpp
@@ -2,11 +2,11 @@
 
 #include "TimingProfilerManager.h"
 
+#include "Framework/Application/SlateApplication.h"
+#include "Framework/Docking/TabManager.h"
 #include "MessageLogModule.h"
 #include "Modules/ModuleManager.h"
 #include "Widgets/Docking/SDockTab.h"
-#include "WorkspaceMenuStructure.h"
-#include "WorkspaceMenuStructureModule.h"
 
 // TraceServices
 #include "TraceServices/Model/Counters.h"
@@ -17,6 +17,7 @@
 #include "Insights/InsightsSettings.h"
 #include "Insights/InsightsStyle.h"
 #include "Insights/TimingProfiler/ViewModels/TimerButterflyAggregator.h"
+#include "Insights/TimingProfiler/ViewModels/TimerNode.h"
 #include "Insights/TimingProfiler/ViewModels/TimingExporter.h"
 #include "Insights/TimingProfiler/Widgets/SFrameTrack.h"
 #include "Insights/TimingProfiler/Widgets/SStatsView.h"
@@ -24,6 +25,7 @@
 #include "Insights/TimingProfiler/Widgets/STimerTreeView.h"
 #include "Insights/TimingProfiler/Widgets/STimingProfilerWindow.h"
 #include "Insights/TimingProfilerCommon.h"
+#include "Insights/ViewModels/LogFilter.h"
 #include "Insights/Widgets/SLogView.h"
 #include "Insights/Widgets/STimingView.h"
 
@@ -353,6 +355,37 @@ void FTimingProfilerManager::FinishTimerButterflyAggregation()
 
 ////////////////////////////////////////////////////////////////////////////////////////////////////
 
+void FTimingProfilerManager::ActivateMajorTab()
+{
+	// Ensure Timing Insights / Timing View is the active tab / view.
+	TSharedPtr<SDockTab> TimingInsightsTab = FGlobalTabmanager::Get()->FindExistingLiveTab(FInsightsManagerTabs::TimingProfilerTabId);
+	if (TimingInsightsTab.IsValid())
+	{
+		TimingInsightsTab->ActivateInParent(ETabActivationCause::SetDirectly);
+		ActivateView();
+	}
+}
+
+////////////////////////////////////////////////////////////////////////////////////////////////////
+
+void FTimingProfilerManager::ActivateView()
+{
+	// Ensure Timing View is the active view.
+	TSharedPtr<STimingProfilerWindow> Wnd = GetProfilerWindow();
+	if (Wnd.IsValid())
+	{
+		TSharedPtr<FTabManager> TabManager = Wnd->GetTabManager();
+		TSharedPtr<SDockTab> TimingViewTab = TabManager->FindExistingLiveTab(FTimingProfilerTabs::TimingViewID);
+		if (TimingViewTab.IsValid())
+		{
+			TimingViewTab->ActivateInParent(ETabActivationCause::SetDirectly);
+			FSlateApplication::Get().SetKeyboardFocus(TimingViewTab->GetContent());
+		}
+	}
+}
+
+////////////////////////////////////////////////////////////////////////////////////////////////////
+
 void FTimingProfilerManager::OnSessionChanged()
 {
 	UE_LOG(LogTimingProfiler, Log, TEXT("OnSessionChanged"));
@@ -488,7 +521,7 @@ void FTimingProfilerManager::ShowHideLogView(const bool bIsVisible)
 		Wnd->ShowHideTab(FTimingProfilerTabs::LogViewID, bIsLogViewVisible);
 	}
 }
-	
+
 ////////////////////////////////////////////////////////////////////////////////////////////////////
 
 void FTimingProfilerManager::ShowHideModulesView(const bool bIsVisible)
diff --git a/Engine/Source/Developer/TraceInsights/Private/Insights/TimingProfiler/TimingProfilerManager.h b/Engine/Source/Developer/TraceInsights/Private/Insights/TimingProfiler/TimingProfilerManager.h
--- a/Engine/Source/Developer/TraceInsights/Private/Insights/TimingProfiler/TimingProfilerManager.h
+++ b/Engine/Source/Developer/TraceInsights/Private/Insights/TimingProfiler/TimingProfilerManager.h
@@ -6,13 +6,15 @@
 
 #include "Containers/Ticker.h"
 #include "Framework/Commands/UICommandList.h"
+#include "Templates/SharedPointer.h"
+
+// TraceInsightsCore
+#include "InsightsCore/Common/AvailabilityCheck.h"
 
 // TraceInsights
-#include "Insights/InsightsManager.h"
 #include "Insights/IUnrealInsightsModule.h"
 #include "Insights/TimingProfiler/TimingProfilerCommands.h"
-#include "Insights/TimingProfiler/ViewModels/TimerNode.h"
-#include "Insights/ViewModels/LogFilter.h"
+#include "Insights/TimingProfilerCommon.h"
 
 namespace UE::Insights
 {
@@ -26,11 +28,14 @@ namespace UE::Insights
 
 		Count
 	};
+
+	class FLogFilter;
 }
 
 namespace UE::Insights::TimingProfiler
 {
 
+class FTimerNode;
 class FTimerButterflyAggregator;
 class STimingProfilerWindow;
 
@@ -95,6 +100,8 @@ public:
 		return ProfilerWindowWeakPtr.Pin();
 	}
 
+	const FName& GetLogListingName() const { return LogListingName; }
+
 	const TSharedRef<FLogFilter>& GetLogFilter() const { return Filter; }
 
 	////////////////////////////////////////////////////////////////////////////////////////////////////
@@ -142,6 +149,9 @@ public:
 
 	////////////////////////////////////////////////////////////////////////////////////////////////////
 
+	void ActivateMajorTab();
+	void ActivateView();
+
 	void OnSessionChanged();
 
 	bool IsValidTimeSelection() const { return SelectionStartTime < SelectionEndTime; }
@@ -149,7 +159,7 @@ public:
 	double GetSelectionEndTime() const { return SelectionEndTime; }
 	void SetSelectedTimeRange(double StartTime, double EndTime);
 
-	FTimerNodePtr GetTimerNode(uint32 TimerId) const;
+	TSharedPtr<FTimerNode> GetTimerNode(uint32 TimerId) const;
 	static constexpr uint32 GetInvalidTimerId() { return InvalidTimerId; }
 	bool IsValidSelectedTimer() const { return SelectedTimerId != InvalidTimerId; }
 	uint32 GetSelectedTimer() const { return SelectedTimerId; }
@@ -165,8 +175,6 @@ public:
 	void UpdateAggregatedTimerStats();
 	void UpdateAggregatedCounterStats();
 
-	const FName& GetLogListingName() const { return LogListingName; }
-
 	ETimingEventsColoringMode GetColoringMode() const { return ColoringMode; }
 	void SetColoringMode(ETimingEventsColoringMode InColoringMode) { ColoringMode = InColoringMode; }
 
diff --git a/Engine/Source/Developer/TraceInsights/Private/Insights/TimingProfiler/ViewModels/TimerFilters.cpp b/Engine/Source/Developer/TraceInsights/Private/Insights/TimingProfiler/ViewModels/TimerFilters.cpp
--- a/Engine/Source/Developer/TraceInsights/Private/Insights/TimingProfiler/ViewModels/TimerFilters.cpp
+++ b/Engine/Source/Developer/TraceInsights/Private/Insights/TimingProfiler/ViewModels/TimerFilters.cpp
@@ -6,7 +6,11 @@
 #include "Widgets/Input/SEditableTextBox.h"
 #include "Widgets/SBoxPanel.h"
 
+// TraceServices
+#include "TraceServices/Model/TimingProfiler.h"
+
 // TraceInsights
+#include "Insights/InsightsManager.h"
 #include "Insights/TimingProfiler/TimingProfilerManager.h"
 #include "Insights/TimingProfiler/Tracks/ThreadTimingTrack.h"
 
