[Insights]
- Removed unused HighlightedNodeName from all tree/table views.
- Fixed incorrect checking for bound delegate (GetColumnOutlineHAlignmentDelegate) for several table views.
- Asset Loading Insights: Fixed possible crash when using an invalid color schema index.

#Synced-CL 51163422

diff --git a/Engine/Source/Developer/TraceInsights/Private/Insights/LoadingProfiler/ViewModels/LoadingSharedState.cpp b/Engine/Source/Developer/TraceInsights/Private/Insights/LoadingProfiler/ViewModels/LoadingSharedState.cpp
--- a/Engine/Source/Developer/TraceInsights/Private/Insights/LoadingProfiler/ViewModels/LoadingSharedState.cpp
+++ b/Engine/Source/Developer/TraceInsights/Private/Insights/LoadingProfiler/ViewModels/LoadingSharedState.cpp
@@ -222,7 +222,11 @@ const TCHAR* FLoadingSharedState::GetEventNameByPackageAndExportClassName(uint32
 
 const TCHAR* FLoadingSharedState::GetEventName(uint32 Depth, const TraceServices::FLoadTimeProfilerCpuEvent& Event) const
 {
-	return GetEventNameDelegate.Execute(Depth, Event);
+	if (GetEventNameDelegate.IsBound())
+	{
+		return GetEventNameDelegate.Execute(Depth, Event);
+	}
+	return TEXT("");
 }
 
 ////////////////////////////////////////////////////////////////////////////////////////////////////
@@ -235,6 +239,7 @@ void FLoadingSharedState::SetColorSchema(int32 Schema)
 	case 1: GetEventNameDelegate = FLoadingTrackGetEventNameDelegate::CreateRaw(this, &FLoadingSharedState::GetEventNameByPackageName); break;
 	case 2: GetEventNameDelegate = FLoadingTrackGetEventNameDelegate::CreateRaw(this, &FLoadingSharedState::GetEventNameByExportClassName); break;
 	case 3: GetEventNameDelegate = FLoadingTrackGetEventNameDelegate::CreateRaw(this, &FLoadingSharedState::GetEventNameByPackageAndExportClassName); break;
+	default: GetEventNameDelegate.Unbind();
 	};
 
 	for (const auto& KV : LoadingTracks)
diff --git a/Engine/Source/Developer/TraceInsights/Private/Insights/NetworkingProfiler/Widgets/SNetStatsCountersTableRow.cpp b/Engine/Source/Developer/TraceInsights/Private/Insights/NetworkingProfiler/Widgets/SNetStatsCountersTableRow.cpp
--- a/Engine/Source/Developer/TraceInsights/Private/Insights/NetworkingProfiler/Widgets/SNetStatsCountersTableRow.cpp
+++ b/Engine/Source/Developer/TraceInsights/Private/Insights/NetworkingProfiler/Widgets/SNetStatsCountersTableRow.cpp
@@ -33,7 +33,6 @@ void SNetStatsCountersTableRow::Construct(const FArguments& InArgs, const TShare
 	SetHoveredCellDelegate = InArgs._OnSetHoveredCell;
 
 	HighlightText = InArgs._HighlightText;
-	HighlightedNodeName = InArgs._HighlightedNodeName;
 
 	TablePtr = InArgs._TablePtr;
 	NetStatsCounterNodePtr = InArgs._NetStatsCounterNodePtr;
@@ -129,10 +128,7 @@ FSlateColor SNetStatsCountersTableRow::GetBackgroundColorAndOpacity(uint32 Size)
 
 FSlateColor SNetStatsCountersTableRow::GetOutlineColorAndOpacity() const
 {
-	const FLinearColor NoColor(0.0f, 0.0f, 0.0f, 0.0f);
-	const bool bShouldBeHighlighted = NetStatsCounterNodePtr->GetName() == HighlightedNodeName.Get();
-	const FLinearColor OutlineColorAndOpacity = bShouldBeHighlighted ? FLinearColor(FColorList::SlateBlue) : NoColor;
-	return OutlineColorAndOpacity;
+	return FLinearColor(0.0f, 0.0f, 0.0f, 0.0f);
 }
 
 ////////////////////////////////////////////////////////////////////////////////////////////////////
@@ -140,11 +136,10 @@ FSlateColor SNetStatsCountersTableRow::GetOutlineColorAndOpacity() const
 const FSlateBrush* SNetStatsCountersTableRow::GetOutlineBrush(const FName ColumnId) const
 {
 	EHorizontalAlignment HAlign = HAlign_Center;
-	if (IsColumnVisibleDelegate.IsBound())
+	if (GetColumnOutlineHAlignmentDelegate.IsBound())
 	{
 		HAlign = GetColumnOutlineHAlignmentDelegate.Execute(ColumnId);
 	}
-
 	return FInsightsStyle::GetOutlineBrush(HAlign);
 }
 
@@ -152,22 +147,15 @@ const FSlateBrush* SNetStatsCountersTableRow::GetOutlineBrush(const FName Column
 
 bool SNetStatsCountersTableRow::HandleShouldBeEnabled() const
 {
-	bool bResult = false;
-
 	if (NetStatsCounterNodePtr->IsGroup())
 	{
-		bResult = true;
+		return true;
 	}
-	else
+	if (OnShouldBeEnabled.IsBound())
 	{
-		if (OnShouldBeEnabled.IsBound())
-		{
-			bResult = OnShouldBeEnabled.Execute(NetStatsCounterNodePtr);
-		}
-		return true;
+		return OnShouldBeEnabled.Execute(NetStatsCounterNodePtr);
 	}
-
-	return bResult;
+	return false;
 }
 
 ////////////////////////////////////////////////////////////////////////////////////////////////////
@@ -175,12 +163,10 @@ bool SNetStatsCountersTableRow::HandleShouldBeEnabled() const
 EVisibility SNetStatsCountersTableRow::IsColumnVisible(const FName ColumnId) const
 {
 	EVisibility Result = EVisibility::Collapsed;
-
 	if (IsColumnVisibleDelegate.IsBound())
 	{
 		Result = IsColumnVisibleDelegate.Execute(ColumnId) ? EVisibility::Visible : EVisibility::Collapsed;
 	}
-
 	return Result;
 }
 
diff --git a/Engine/Source/Developer/TraceInsights/Private/Insights/NetworkingProfiler/Widgets/SNetStatsCountersTableRow.h b/Engine/Source/Developer/TraceInsights/Private/Insights/NetworkingProfiler/Widgets/SNetStatsCountersTableRow.h
--- a/Engine/Source/Developer/TraceInsights/Private/Insights/NetworkingProfiler/Widgets/SNetStatsCountersTableRow.h
+++ b/Engine/Source/Developer/TraceInsights/Private/Insights/NetworkingProfiler/Widgets/SNetStatsCountersTableRow.h
@@ -41,7 +41,6 @@ public:
 		SLATE_EVENT(FGetColumnOutlineHAlignmentDelegate, OnGetColumnOutlineHAlignmentDelegate)
 		SLATE_EVENT(FSetHoveredNetStatsCounterTableCell, OnSetHoveredCell)
 		SLATE_ATTRIBUTE(FText, HighlightText)
-		SLATE_ATTRIBUTE(FName, HighlightedNodeName)
 		SLATE_ARGUMENT(TSharedPtr<FTable>, TablePtr)
 		SLATE_ARGUMENT(FNetStatsCounterNodePtr, NetStatsCounterNodePtr)
 	SLATE_END_ARGS()
@@ -94,9 +93,6 @@ protected:
 	/** Text to be highlighted on timer name. */
 	TAttribute<FText> HighlightText;
 
-	/** Name of the timer node that should be drawn as highlighted. */
-	TAttribute<FName> HighlightedNodeName;
-
 	TSharedPtr<SNetStatsCounterTableRowToolTip> RowToolTip;
 };
 
diff --git a/Engine/Source/Developer/TraceInsights/Private/Insights/NetworkingProfiler/Widgets/SNetStatsCountersView.cpp b/Engine/Source/Developer/TraceInsights/Private/Insights/NetworkingProfiler/Widgets/SNetStatsCountersView.cpp
--- a/Engine/Source/Developer/TraceInsights/Private/Insights/NetworkingProfiler/Widgets/SNetStatsCountersView.cpp
+++ b/Engine/Source/Developer/TraceInsights/Private/Insights/NetworkingProfiler/Widgets/SNetStatsCountersView.cpp
@@ -821,7 +821,6 @@ TSharedRef<ITableRow> SNetStatsCountersView::TreeView_OnGenerateRow(FNetStatsCou
 		.OnSetHoveredCell(this, &SNetStatsCountersView::TableRow_SetHoveredCell)
 		.OnGetColumnOutlineHAlignmentDelegate(this, &SNetStatsCountersView::TableRow_GetColumnOutlineHAlignment)
 		.HighlightText(this, &SNetStatsCountersView::TableRow_GetHighlightText)
-		.HighlightedNodeName(this, &SNetStatsCountersView::TableRow_GetHighlightedNodeName)
 		.TablePtr(Table)
 		.NetStatsCounterNodePtr(NetStatsCounterNodePtr);
 
@@ -853,15 +852,15 @@ void SNetStatsCountersView::TableRow_SetHoveredCell(TSharedPtr<FTable> InTablePt
 EHorizontalAlignment SNetStatsCountersView::TableRow_GetColumnOutlineHAlignment(const FName ColumnId) const
 {
 	const TIndirectArray<SHeaderRow::FColumn>& Columns = TreeViewHeaderRow->GetColumns();
-	const int32 LastColumnIdx = Columns.Num() - 1;
 
 	// First column
-	if (Columns[0].ColumnId == ColumnId)
+	if (Columns.Num() == 0 ||
+		Columns[0].ColumnId == ColumnId)
 	{
 		return HAlign_Left;
 	}
 	// Last column
-	else if (Columns[LastColumnIdx].ColumnId == ColumnId)
+	else if (Columns[Columns.Num() - 1].ColumnId == ColumnId)
 	{
 		return HAlign_Right;
 	}
@@ -878,13 +877,6 @@ FText SNetStatsCountersView::TableRow_GetHighlightText() const
 	return SearchBox->GetText();
 }
 
-////////////////////////////////////////////////////////////////////////////////////////////////////
-
-FName SNetStatsCountersView::TableRow_GetHighlightedNodeName() const
-{
-	return HighlightedNodeName;
-}
-
 ////////////////////////////////////////////////////////////////////////////////////////////////////
 // SearchBox
 ////////////////////////////////////////////////////////////////////////////////////////////////////
diff --git a/Engine/Source/Developer/TraceInsights/Private/Insights/NetworkingProfiler/Widgets/SNetStatsCountersView.h b/Engine/Source/Developer/TraceInsights/Private/Insights/NetworkingProfiler/Widgets/SNetStatsCountersView.h
--- a/Engine/Source/Developer/TraceInsights/Private/Insights/NetworkingProfiler/Widgets/SNetStatsCountersView.h
+++ b/Engine/Source/Developer/TraceInsights/Private/Insights/NetworkingProfiler/Widgets/SNetStatsCountersView.h
@@ -161,7 +161,6 @@ protected:
 	EHorizontalAlignment TableRow_GetColumnOutlineHAlignment(const FName ColumnId) const;
 
 	FText TableRow_GetHighlightText() const;
-	FName TableRow_GetHighlightedNodeName() const;
 
 	////////////////////////////////////////////////////////////////////////////////////////////////////
 	// Filtering
@@ -287,9 +286,6 @@ protected:
 	/** A shared pointer to the node currently being hovered by the mouse. */
 	FNetStatsCounterNodePtr HoveredNodePtr;
 
-	/** Name of the node that should be drawn as highlighted. */
-	FName HighlightedNodeName;
-
 	//////////////////////////////////////////////////
 	// NetStatsCounters Node
 
diff --git a/Engine/Source/Developer/TraceInsights/Private/Insights/NetworkingProfiler/Widgets/SNetStatsTableRow.cpp b/Engine/Source/Developer/TraceInsights/Private/Insights/NetworkingProfiler/Widgets/SNetStatsTableRow.cpp
--- a/Engine/Source/Developer/TraceInsights/Private/Insights/NetworkingProfiler/Widgets/SNetStatsTableRow.cpp
+++ b/Engine/Source/Developer/TraceInsights/Private/Insights/NetworkingProfiler/Widgets/SNetStatsTableRow.cpp
@@ -33,7 +33,6 @@ void SNetStatsTableRow::Construct(const FArguments& InArgs, const TSharedRef<STa
 	SetHoveredCellDelegate = InArgs._OnSetHoveredCell;
 
 	HighlightText = InArgs._HighlightText;
-	HighlightedNodeName = InArgs._HighlightedNodeName;
 
 	TablePtr = InArgs._TablePtr;
 	NetEventNodePtr = InArgs._NetEventNodePtr;
@@ -129,10 +128,7 @@ FSlateColor SNetStatsTableRow::GetBackgroundColorAndOpacity(uint32 Size) const
 
 FSlateColor SNetStatsTableRow::GetOutlineColorAndOpacity() const
 {
-	const FLinearColor NoColor(0.0f, 0.0f, 0.0f, 0.0f);
-	const bool bShouldBeHighlighted = NetEventNodePtr->GetName() == HighlightedNodeName.Get();
-	const FLinearColor OutlineColorAndOpacity = bShouldBeHighlighted ? FLinearColor(FColorList::SlateBlue) : NoColor;
-	return OutlineColorAndOpacity;
+	return FLinearColor(0.0f, 0.0f, 0.0f, 0.0f);
 }
 
 ////////////////////////////////////////////////////////////////////////////////////////////////////
@@ -140,11 +136,10 @@ FSlateColor SNetStatsTableRow::GetOutlineColorAndOpacity() const
 const FSlateBrush* SNetStatsTableRow::GetOutlineBrush(const FName ColumnId) const
 {
 	EHorizontalAlignment HAlign = HAlign_Center;
-	if (IsColumnVisibleDelegate.IsBound())
+	if (GetColumnOutlineHAlignmentDelegate.IsBound())
 	{
 		HAlign = GetColumnOutlineHAlignmentDelegate.Execute(ColumnId);
 	}
-
 	return FInsightsStyle::GetOutlineBrush(HAlign);
 }
 
@@ -152,21 +147,15 @@ const FSlateBrush* SNetStatsTableRow::GetOutlineBrush(const FName ColumnId) cons
 
 bool SNetStatsTableRow::HandleShouldBeEnabled() const
 {
-	bool bResult = false;
-
 	if (NetEventNodePtr->IsGroup())
 	{
-		bResult = true;
+		return true;
 	}
-	else
+	if (OnShouldBeEnabled.IsBound())
 	{
-		if (OnShouldBeEnabled.IsBound())
-		{
-			bResult = OnShouldBeEnabled.Execute(NetEventNodePtr);
-		}
+		return OnShouldBeEnabled.Execute(NetEventNodePtr);
 	}
-
-	return bResult;
+	return false;
 }
 
 ////////////////////////////////////////////////////////////////////////////////////////////////////
@@ -174,12 +163,10 @@ bool SNetStatsTableRow::HandleShouldBeEnabled() const
 EVisibility SNetStatsTableRow::IsColumnVisible(const FName ColumnId) const
 {
 	EVisibility Result = EVisibility::Collapsed;
-
 	if (IsColumnVisibleDelegate.IsBound())
 	{
 		Result = IsColumnVisibleDelegate.Execute(ColumnId) ? EVisibility::Visible : EVisibility::Collapsed;
 	}
-
 	return Result;
 }
 
diff --git a/Engine/Source/Developer/TraceInsights/Private/Insights/NetworkingProfiler/Widgets/SNetStatsTableRow.h b/Engine/Source/Developer/TraceInsights/Private/Insights/NetworkingProfiler/Widgets/SNetStatsTableRow.h
--- a/Engine/Source/Developer/TraceInsights/Private/Insights/NetworkingProfiler/Widgets/SNetStatsTableRow.h
+++ b/Engine/Source/Developer/TraceInsights/Private/Insights/NetworkingProfiler/Widgets/SNetStatsTableRow.h
@@ -42,7 +42,6 @@ public:
 		SLATE_EVENT(FGetColumnOutlineHAlignmentDelegate, OnGetColumnOutlineHAlignmentDelegate)
 		SLATE_EVENT(FSetHoveredNetEventTableCell, OnSetHoveredCell)
 		SLATE_ATTRIBUTE(FText, HighlightText)
-		SLATE_ATTRIBUTE(FName, HighlightedNodeName)
 		SLATE_ARGUMENT(TSharedPtr<FTable>, TablePtr)
 		SLATE_ARGUMENT(FNetEventNodePtr, NetEventNodePtr)
 	SLATE_END_ARGS()
@@ -95,9 +94,6 @@ protected:
 	/** Text to be highlighted on timer name. */
 	TAttribute<FText> HighlightText;
 
-	/** Name of the timer node that should be drawn as highlighted. */
-	TAttribute<FName> HighlightedNodeName;
-
 	TSharedPtr<SNetEventTableRowToolTip> RowToolTip;
 };
 
diff --git a/Engine/Source/Developer/TraceInsights/Private/Insights/NetworkingProfiler/Widgets/SNetStatsView.cpp b/Engine/Source/Developer/TraceInsights/Private/Insights/NetworkingProfiler/Widgets/SNetStatsView.cpp
--- a/Engine/Source/Developer/TraceInsights/Private/Insights/NetworkingProfiler/Widgets/SNetStatsView.cpp
+++ b/Engine/Source/Developer/TraceInsights/Private/Insights/NetworkingProfiler/Widgets/SNetStatsView.cpp
@@ -809,7 +809,6 @@ TSharedRef<ITableRow> SNetStatsView::TreeView_OnGenerateRow(FNetEventNodePtr Net
 		.OnSetHoveredCell(this, &SNetStatsView::TableRow_SetHoveredCell)
 		.OnGetColumnOutlineHAlignmentDelegate(this, &SNetStatsView::TableRow_GetColumnOutlineHAlignment)
 		.HighlightText(this, &SNetStatsView::TableRow_GetHighlightText)
-		.HighlightedNodeName(this, &SNetStatsView::TableRow_GetHighlightedNodeName)
 		.TablePtr(Table)
 		.NetEventNodePtr(NetEventNodePtr);
 
@@ -841,15 +840,15 @@ void SNetStatsView::TableRow_SetHoveredCell(TSharedPtr<FTable> InTablePtr, TShar
 EHorizontalAlignment SNetStatsView::TableRow_GetColumnOutlineHAlignment(const FName ColumnId) const
 {
 	const TIndirectArray<SHeaderRow::FColumn>& Columns = TreeViewHeaderRow->GetColumns();
-	const int32 LastColumnIdx = Columns.Num() - 1;
 
 	// First column
-	if (Columns[0].ColumnId == ColumnId)
+	if (Columns.Num() == 0 ||
+		Columns[0].ColumnId == ColumnId)
 	{
 		return HAlign_Left;
 	}
 	// Last column
-	else if (Columns[LastColumnIdx].ColumnId == ColumnId)
+	else if (Columns[Columns.Num() - 1].ColumnId == ColumnId)
 	{
 		return HAlign_Right;
 	}
@@ -866,13 +865,6 @@ FText SNetStatsView::TableRow_GetHighlightText() const
 	return SearchBox->GetText();
 }
 
-////////////////////////////////////////////////////////////////////////////////////////////////////
-
-FName SNetStatsView::TableRow_GetHighlightedNodeName() const
-{
-	return HighlightedNodeName;
-}
-
 ////////////////////////////////////////////////////////////////////////////////////////////////////
 // SearchBox
 ////////////////////////////////////////////////////////////////////////////////////////////////////
diff --git a/Engine/Source/Developer/TraceInsights/Private/Insights/NetworkingProfiler/Widgets/SNetStatsView.h b/Engine/Source/Developer/TraceInsights/Private/Insights/NetworkingProfiler/Widgets/SNetStatsView.h
--- a/Engine/Source/Developer/TraceInsights/Private/Insights/NetworkingProfiler/Widgets/SNetStatsView.h
+++ b/Engine/Source/Developer/TraceInsights/Private/Insights/NetworkingProfiler/Widgets/SNetStatsView.h
@@ -168,8 +168,7 @@ protected:
 	EHorizontalAlignment TableRow_GetColumnOutlineHAlignment(const FName ColumnId) const;
 
 	FText TableRow_GetHighlightText() const;
-	FName TableRow_GetHighlightedNodeName() const;
-
+	
 	////////////////////////////////////////////////////////////////////////////////////////////////////
 	// Filtering
 
@@ -295,9 +294,6 @@ protected:
 	/** A shared pointer to the node currently being hovered by the mouse. */
 	FNetEventNodePtr HoveredNodePtr;
 
-	/** Name of the node that should be drawn as highlighted. */
-	FName HighlightedNodeName;
-
 	//////////////////////////////////////////////////
 	// Net Event Nodes
 
diff --git a/Engine/Source/Developer/TraceInsights/Private/Insights/TimingProfiler/Widgets/SStatsTableRow.cpp b/Engine/Source/Developer/TraceInsights/Private/Insights/TimingProfiler/Widgets/SStatsTableRow.cpp
--- a/Engine/Source/Developer/TraceInsights/Private/Insights/TimingProfiler/Widgets/SStatsTableRow.cpp
+++ b/Engine/Source/Developer/TraceInsights/Private/Insights/TimingProfiler/Widgets/SStatsTableRow.cpp
@@ -34,7 +34,6 @@ void SStatsTableRow::Construct(const FArguments& InArgs, const TSharedRef<STable
 	SetHoveredCellDelegate = InArgs._OnSetHoveredCell;
 
 	HighlightText = InArgs._HighlightText;
-	HighlightedNodeName = InArgs._HighlightedNodeName;
 
 	TablePtr = InArgs._TablePtr;
 	StatsNodePtr = InArgs._StatsNodePtr;
@@ -142,10 +141,7 @@ FSlateColor SStatsTableRow::GetBackgroundColorAndOpacity(double Time) const
 
 FSlateColor SStatsTableRow::GetOutlineColorAndOpacity() const
 {
-	const FLinearColor NoColor(0.0f, 0.0f, 0.0f, 0.0f);
-	const bool bShouldBeHighlighted = StatsNodePtr->GetName() == HighlightedNodeName.Get();
-	const FLinearColor OutlineColorAndOpacity = bShouldBeHighlighted ? FLinearColor(FColorList::SlateBlue) : NoColor;
-	return OutlineColorAndOpacity;
+	return FLinearColor(0.0f, 0.0f, 0.0f, 0.0f);
 }
 
 ////////////////////////////////////////////////////////////////////////////////////////////////////
@@ -153,7 +149,7 @@ FSlateColor SStatsTableRow::GetOutlineColorAndOpacity() const
 const FSlateBrush* SStatsTableRow::GetOutlineBrush(const FName ColumnId) const
 {
 	EHorizontalAlignment HAlign = HAlign_Center;
-	if (IsColumnVisibleDelegate.IsBound())
+	if (GetColumnOutlineHAlignmentDelegate.IsBound())
 	{
 		HAlign = GetColumnOutlineHAlignmentDelegate.Execute(ColumnId);
 	}
@@ -164,21 +160,15 @@ const FSlateBrush* SStatsTableRow::GetOutlineBrush(const FName ColumnId) const
 
 bool SStatsTableRow::HandleShouldBeEnabled() const
 {
-	bool bResult = false;
-
 	if (StatsNodePtr->IsGroup())
 	{
-		bResult = true;
+		return true;
 	}
-	else
+	if (OnShouldBeEnabled.IsBound())
 	{
-		if (OnShouldBeEnabled.IsBound())
-		{
-			bResult = OnShouldBeEnabled.Execute(StatsNodePtr);
-		}
+		return OnShouldBeEnabled.Execute(StatsNodePtr);
 	}
-
-	return bResult;
+	return false;
 }
 
 ////////////////////////////////////////////////////////////////////////////////////////////////////
@@ -186,12 +176,10 @@ bool SStatsTableRow::HandleShouldBeEnabled() const
 EVisibility SStatsTableRow::IsColumnVisible(const FName ColumnId) const
 {
 	EVisibility Result = EVisibility::Collapsed;
-
 	if (IsColumnVisibleDelegate.IsBound())
 	{
 		Result = IsColumnVisibleDelegate.Execute(ColumnId) ? EVisibility::Visible : EVisibility::Collapsed;
 	}
-
 	return Result;
 }
 
diff --git a/Engine/Source/Developer/TraceInsights/Private/Insights/TimingProfiler/Widgets/SStatsTableRow.h b/Engine/Source/Developer/TraceInsights/Private/Insights/TimingProfiler/Widgets/SStatsTableRow.h
--- a/Engine/Source/Developer/TraceInsights/Private/Insights/TimingProfiler/Widgets/SStatsTableRow.h
+++ b/Engine/Source/Developer/TraceInsights/Private/Insights/TimingProfiler/Widgets/SStatsTableRow.h
@@ -41,7 +41,6 @@ public:
 		SLATE_EVENT(FGetColumnOutlineHAlignmentDelegate, OnGetColumnOutlineHAlignmentDelegate)
 		SLATE_EVENT(FSetHoveredStatsTableCell, OnSetHoveredCell)
 		SLATE_ATTRIBUTE(FText, HighlightText)
-		SLATE_ATTRIBUTE(FName, HighlightedNodeName)
 		SLATE_ARGUMENT(TSharedPtr<FTable>, TablePtr)
 		SLATE_ARGUMENT(FStatsNodePtr, StatsNodePtr)
 	SLATE_END_ARGS()
@@ -94,9 +93,6 @@ protected:
 	/** Text to be highlighted on stats counter name. */
 	TAttribute<FText> HighlightText;
 
-	/** Name of the stats counter node that should be drawn as highlighted. */
-	TAttribute<FName> HighlightedNodeName;
-
 	TSharedPtr<SStatsCounterTableRowToolTip> RowToolTip;
 };
 
diff --git a/Engine/Source/Developer/TraceInsights/Private/Insights/TimingProfiler/Widgets/SStatsView.cpp b/Engine/Source/Developer/TraceInsights/Private/Insights/TimingProfiler/Widgets/SStatsView.cpp
--- a/Engine/Source/Developer/TraceInsights/Private/Insights/TimingProfiler/Widgets/SStatsView.cpp
+++ b/Engine/Source/Developer/TraceInsights/Private/Insights/TimingProfiler/Widgets/SStatsView.cpp
@@ -1271,7 +1271,6 @@ TSharedRef<ITableRow> SStatsView::TreeView_OnGenerateRow(FStatsNodePtr NodePtr,
 		.OnSetHoveredCell(this, &SStatsView::TableRow_SetHoveredCell)
 		.OnGetColumnOutlineHAlignmentDelegate(this, &SStatsView::TableRow_GetColumnOutlineHAlignment)
 		.HighlightText(this, &SStatsView::TableRow_GetHighlightText)
-		.HighlightedNodeName(this, &SStatsView::TableRow_GetHighlightedNodeName)
 		.TablePtr(Table)
 		.StatsNodePtr(NodePtr);
 
@@ -1303,15 +1302,15 @@ void SStatsView::TableRow_SetHoveredCell(TSharedPtr<FTable> InTablePtr, TSharedP
 EHorizontalAlignment SStatsView::TableRow_GetColumnOutlineHAlignment(const FName ColumnId) const
 {
 	const TIndirectArray<SHeaderRow::FColumn>& Columns = TreeViewHeaderRow->GetColumns();
-	const int32 LastColumnIdx = Columns.Num() - 1;
 
 	// First column
-	if (Columns[0].ColumnId == ColumnId)
+	if (Columns.Num() == 0 ||
+		Columns[0].ColumnId == ColumnId)
 	{
 		return HAlign_Left;
 	}
 	// Last column
-	else if (Columns[LastColumnIdx].ColumnId == ColumnId)
+	else if (Columns[Columns.Num() - 1].ColumnId == ColumnId)
 	{
 		return HAlign_Right;
 	}
@@ -1328,13 +1327,6 @@ FText SStatsView::TableRow_GetHighlightText() const
 	return SearchBox->GetText();
 }
 
-////////////////////////////////////////////////////////////////////////////////////////////////////
-
-FName SStatsView::TableRow_GetHighlightedNodeName() const
-{
-	return HighlightedNodeName;
-}
-
 ////////////////////////////////////////////////////////////////////////////////////////////////////
 // SearchBox
 ////////////////////////////////////////////////////////////////////////////////////////////////////
diff --git a/Engine/Source/Developer/TraceInsights/Private/Insights/TimingProfiler/Widgets/SStatsView.h b/Engine/Source/Developer/TraceInsights/Private/Insights/TimingProfiler/Widgets/SStatsView.h
--- a/Engine/Source/Developer/TraceInsights/Private/Insights/TimingProfiler/Widgets/SStatsView.h
+++ b/Engine/Source/Developer/TraceInsights/Private/Insights/TimingProfiler/Widgets/SStatsView.h
@@ -209,7 +209,6 @@ private:
 	EHorizontalAlignment TableRow_GetColumnOutlineHAlignment(const FName ColumnId) const;
 
 	FText TableRow_GetHighlightText() const;
-	FName TableRow_GetHighlightedNodeName() const;
 
 	////////////////////////////////////////////////////////////////////////////////////////////////////
 	// Filtering
@@ -363,9 +362,6 @@ private:
 	/** A shared pointer to the counter node currently being hovered by the mouse. */
 	FStatsNodePtr HoveredNodePtr;
 
-	/** Name of the counter that should be drawn as highlighted. */
-	FName HighlightedNodeName;
-
 	//////////////////////////////////////////////////
 	// Stats Nodes
 
diff --git a/Engine/Source/Developer/TraceInsights/Private/Insights/TimingProfiler/Widgets/STimerTableRow.cpp b/Engine/Source/Developer/TraceInsights/Private/Insights/TimingProfiler/Widgets/STimerTableRow.cpp
--- a/Engine/Source/Developer/TraceInsights/Private/Insights/TimingProfiler/Widgets/STimerTableRow.cpp
+++ b/Engine/Source/Developer/TraceInsights/Private/Insights/TimingProfiler/Widgets/STimerTableRow.cpp
@@ -34,7 +34,6 @@ void STimerTableRow::Construct(const FArguments& InArgs, const TSharedRef<STable
 	SetHoveredCellDelegate = InArgs._OnSetHoveredCell;
 
 	HighlightText = InArgs._HighlightText;
-	HighlightedNodeName = InArgs._HighlightedNodeName;
 
 	TablePtr = InArgs._TablePtr;
 	TimerNodePtr = InArgs._TimerNodePtr;
@@ -161,10 +160,7 @@ FSlateColor STimerTableRow::GetBackgroundColorAndOpacity(double Time) const
 
 FSlateColor STimerTableRow::GetOutlineColorAndOpacity() const
 {
-	const FLinearColor NoColor(0.0f, 0.0f, 0.0f, 0.0f);
-	const bool bShouldBeHighlighted = TimerNodePtr->GetName() == HighlightedNodeName.Get();
-	const FLinearColor OutlineColorAndOpacity = bShouldBeHighlighted ? FLinearColor(FColorList::SlateBlue) : NoColor;
-	return OutlineColorAndOpacity;
+	return FLinearColor(0.0f, 0.0f, 0.0f, 0.0f);
 }
 
 ////////////////////////////////////////////////////////////////////////////////////////////////////
@@ -172,7 +168,7 @@ FSlateColor STimerTableRow::GetOutlineColorAndOpacity() const
 const FSlateBrush* STimerTableRow::GetOutlineBrush(const FName ColumnId) const
 {
 	EHorizontalAlignment HAlign = HAlign_Center;
-	if (IsColumnVisibleDelegate.IsBound())
+	if (GetColumnOutlineHAlignmentDelegate.IsBound())
 	{
 		HAlign = GetColumnOutlineHAlignmentDelegate.Execute(ColumnId);
 	}
@@ -183,21 +179,15 @@ const FSlateBrush* STimerTableRow::GetOutlineBrush(const FName ColumnId) const
 
 bool STimerTableRow::HandleShouldBeEnabled() const
 {
-	bool bResult = false;
-
 	if (TimerNodePtr->IsGroup())
 	{
-		bResult = true;
+		return true;
 	}
-	else
+	if (OnShouldBeEnabled.IsBound())
 	{
-		if (OnShouldBeEnabled.IsBound())
-		{
-			bResult = OnShouldBeEnabled.Execute(TimerNodePtr);
-		}
+		return OnShouldBeEnabled.Execute(TimerNodePtr);
 	}
-
-	return bResult;
+	return false;
 }
 
 ////////////////////////////////////////////////////////////////////////////////////////////////////
@@ -205,12 +195,10 @@ bool STimerTableRow::HandleShouldBeEnabled() const
 EVisibility STimerTableRow::IsColumnVisible(const FName ColumnId) const
 {
 	EVisibility Result = EVisibility::Collapsed;
-
 	if (IsColumnVisibleDelegate.IsBound())
 	{
 		Result = IsColumnVisibleDelegate.Execute(ColumnId) ? EVisibility::Visible : EVisibility::Collapsed;
 	}
-
 	return Result;
 }
 
diff --git a/Engine/Source/Developer/TraceInsights/Private/Insights/TimingProfiler/Widgets/STimerTableRow.h b/Engine/Source/Developer/TraceInsights/Private/Insights/TimingProfiler/Widgets/STimerTableRow.h
--- a/Engine/Source/Developer/TraceInsights/Private/Insights/TimingProfiler/Widgets/STimerTableRow.h
+++ b/Engine/Source/Developer/TraceInsights/Private/Insights/TimingProfiler/Widgets/STimerTableRow.h
@@ -41,7 +41,6 @@ public:
 		SLATE_EVENT(FGetColumnOutlineHAlignmentDelegate, OnGetColumnOutlineHAlignmentDelegate)
 		SLATE_EVENT(FSetHoveredTimerTableCell, OnSetHoveredCell)
 		SLATE_ATTRIBUTE(FText, HighlightText)
-		SLATE_ATTRIBUTE(FName, HighlightedNodeName)
 		SLATE_ARGUMENT(TSharedPtr<FTable>, TablePtr)
 		SLATE_ARGUMENT(FTimerNodePtr, TimerNodePtr)
 	SLATE_END_ARGS()
@@ -94,9 +93,6 @@ protected:
 	/** Text to be highlighted on timer name. */
 	TAttribute<FText> HighlightText;
 
-	/** Name of the timer node that should be drawn as highlighted. */
-	TAttribute<FName> HighlightedNodeName;
-
 	TSharedPtr<STimerTableRowToolTip> RowToolTip;
 };
 
diff --git a/Engine/Source/Developer/TraceInsights/Private/Insights/TimingProfiler/Widgets/STimerTreeView.cpp b/Engine/Source/Developer/TraceInsights/Private/Insights/TimingProfiler/Widgets/STimerTreeView.cpp
--- a/Engine/Source/Developer/TraceInsights/Private/Insights/TimingProfiler/Widgets/STimerTreeView.cpp
+++ b/Engine/Source/Developer/TraceInsights/Private/Insights/TimingProfiler/Widgets/STimerTreeView.cpp
@@ -1061,7 +1061,6 @@ TSharedRef<ITableRow> STimerTreeView::TreeView_OnGenerateRow(FTimerNodePtr NodeP
 		.OnSetHoveredCell(this, &STimerTreeView::TableRow_SetHoveredCell)
 		.OnGetColumnOutlineHAlignmentDelegate(this, &STimerTreeView::TableRow_GetColumnOutlineHAlignment)
 		.HighlightText(this, &STimerTreeView::TableRow_GetHighlightText)
-		.HighlightedNodeName(this, &STimerTreeView::TableRow_GetHighlightedNodeName)
 		.TablePtr(Table)
 		.TimerNodePtr(NodePtr);
 
@@ -1093,15 +1092,15 @@ void STimerTreeView::TableRow_SetHoveredCell(TSharedPtr<FTable> InTablePtr, TSha
 EHorizontalAlignment STimerTreeView::TableRow_GetColumnOutlineHAlignment(const FName ColumnId) const
 {
 	const TIndirectArray<SHeaderRow::FColumn>& Columns = TreeViewHeaderRow->GetColumns();
-	const int32 LastColumnIdx = Columns.Num() - 1;
 
 	// First column
-	if (Columns[0].ColumnId == ColumnId)
+	if (Columns.Num() == 0 ||
+		Columns[0].ColumnId == ColumnId)
 	{
 		return HAlign_Left;
 	}
 	// Last column
-	else if (Columns[LastColumnIdx].ColumnId == ColumnId)
+	else if (Columns[Columns.Num() - 1].ColumnId == ColumnId)
 	{
 		return HAlign_Right;
 	}
@@ -1118,13 +1117,6 @@ FText STimerTreeView::TableRow_GetHighlightText() const
 	return FText::GetEmpty();
 }
 
-////////////////////////////////////////////////////////////////////////////////////////////////////
-
-FName STimerTreeView::TableRow_GetHighlightedNodeName() const
-{
-	return HighlightedNodeName;
-}
-
 ////////////////////////////////////////////////////////////////////////////////////////////////////
 // Sorting
 ////////////////////////////////////////////////////////////////////////////////////////////////////
diff --git a/Engine/Source/Developer/TraceInsights/Private/Insights/TimingProfiler/Widgets/STimerTreeView.h b/Engine/Source/Developer/TraceInsights/Private/Insights/TimingProfiler/Widgets/STimerTreeView.h
--- a/Engine/Source/Developer/TraceInsights/Private/Insights/TimingProfiler/Widgets/STimerTreeView.h
+++ b/Engine/Source/Developer/TraceInsights/Private/Insights/TimingProfiler/Widgets/STimerTreeView.h
@@ -174,8 +174,7 @@ private:
 	EHorizontalAlignment TableRow_GetColumnOutlineHAlignment(const FName ColumnId) const;
 
 	FText TableRow_GetHighlightText() const;
-	FName TableRow_GetHighlightedNodeName() const;
-
+	
 	////////////////////////////////////////////////////////////////////////////////////////////////////
 	// Sorting
 
@@ -295,9 +294,6 @@ private:
 	/** A shared pointer to the timer node currently being hovered by the mouse. */
 	FTimerNodePtr HoveredNodePtr;
 
-	/** Name of the timer that should be drawn as highlighted. */
-	FName HighlightedNodeName;
-
 	//////////////////////////////////////////////////
 	// Timer Nodes
 
diff --git a/Engine/Source/Developer/TraceInsights/Private/Insights/TimingProfiler/Widgets/STimersView.cpp b/Engine/Source/Developer/TraceInsights/Private/Insights/TimingProfiler/Widgets/STimersView.cpp
--- a/Engine/Source/Developer/TraceInsights/Private/Insights/TimingProfiler/Widgets/STimersView.cpp
+++ b/Engine/Source/Developer/TraceInsights/Private/Insights/TimingProfiler/Widgets/STimersView.cpp
@@ -1887,7 +1887,6 @@ TSharedRef<ITableRow> STimersView::TreeView_OnGenerateRow(FTimerNodePtr NodePtr,
 		.OnSetHoveredCell(this, &STimersView::TableRow_SetHoveredCell)
 		.OnGetColumnOutlineHAlignmentDelegate(this, &STimersView::TableRow_GetColumnOutlineHAlignment)
 		.HighlightText(this, &STimersView::TableRow_GetHighlightText)
-		.HighlightedNodeName(this, &STimersView::TableRow_GetHighlightedNodeName)
 		.TablePtr(Table)
 		.TimerNodePtr(NodePtr);
 
@@ -1919,15 +1918,15 @@ void STimersView::TableRow_SetHoveredCell(TSharedPtr<FTable> InTablePtr, TShared
 EHorizontalAlignment STimersView::TableRow_GetColumnOutlineHAlignment(const FName ColumnId) const
 {
 	const TIndirectArray<SHeaderRow::FColumn>& Columns = TreeViewHeaderRow->GetColumns();
-	const int32 LastColumnIdx = Columns.Num() - 1;
 
 	// First column
-	if (Columns[0].ColumnId == ColumnId)
+	if (Columns.Num() == 0 ||
+		Columns[0].ColumnId == ColumnId)
 	{
 		return HAlign_Left;
 	}
 	// Last column
-	else if (Columns[LastColumnIdx].ColumnId == ColumnId)
+	else if (Columns[Columns.Num() - 1].ColumnId == ColumnId)
 	{
 		return HAlign_Right;
 	}
@@ -1944,13 +1943,6 @@ FText STimersView::TableRow_GetHighlightText() const
 	return SearchBox->GetText();
 }
 
-////////////////////////////////////////////////////////////////////////////////////////////////////
-
-FName STimersView::TableRow_GetHighlightedNodeName() const
-{
-	return HighlightedNodeName;
-}
-
 ////////////////////////////////////////////////////////////////////////////////////////////////////
 // SearchBox
 ////////////////////////////////////////////////////////////////////////////////////////////////////
diff --git a/Engine/Source/Developer/TraceInsights/Private/Insights/TimingProfiler/Widgets/STimersView.h b/Engine/Source/Developer/TraceInsights/Private/Insights/TimingProfiler/Widgets/STimersView.h
--- a/Engine/Source/Developer/TraceInsights/Private/Insights/TimingProfiler/Widgets/STimersView.h
+++ b/Engine/Source/Developer/TraceInsights/Private/Insights/TimingProfiler/Widgets/STimersView.h
@@ -253,8 +253,7 @@ private:
 	EHorizontalAlignment TableRow_GetColumnOutlineHAlignment(const FName ColumnId) const;
 
 	FText TableRow_GetHighlightText() const;
-	FName TableRow_GetHighlightedNodeName() const;
-
+	
 	////////////////////////////////////////////////////////////////////////////////////////////////////
 	// Filtering
 
@@ -446,9 +445,6 @@ private:
 	/** A shared pointer to the timer node currently being hovered by the mouse. */
 	FTimerNodePtr HoveredNodePtr;
 
-	/** Name of the timer that should be drawn as highlighted. */
-	FName HighlightedNodeName;
-
 	//////////////////////////////////////////////////
 	// Timer Nodes
 
diff --git a/Engine/Source/Developer/TraceInsightsCore/Private/InsightsCore/Table/Widgets/STableTreeView.cpp b/Engine/Source/Developer/TraceInsightsCore/Private/InsightsCore/Table/Widgets/STableTreeView.cpp
--- a/Engine/Source/Developer/TraceInsightsCore/Private/InsightsCore/Table/Widgets/STableTreeView.cpp
+++ b/Engine/Source/Developer/TraceInsightsCore/Private/InsightsCore/Table/Widgets/STableTreeView.cpp
@@ -1706,7 +1706,6 @@ TSharedRef<ITableRow> STableTreeView::TreeView_OnGenerateRow(FTableTreeNodePtr N
 		.OnSetHoveredCell(this, &STableTreeView::TableRow_SetHoveredCell)
 		.OnGetColumnOutlineHAlignmentDelegate(this, &STableTreeView::TableRow_GetColumnOutlineHAlignment)
 		.HighlightText(this, &STableTreeView::TableRow_GetHighlightText)
-		.HighlightedNodeName(this, &STableTreeView::TableRow_GetHighlightedNodeName)
 		.TablePtr(Table)
 		.TableTreeNodePtr(NodePtr);
 
@@ -1737,15 +1736,15 @@ void STableTreeView::TableRow_SetHoveredCell(TSharedPtr<FTable> InTablePtr, TSha
 EHorizontalAlignment STableTreeView::TableRow_GetColumnOutlineHAlignment(const FName ColumnId) const
 {
 	const TIndirectArray<SHeaderRow::FColumn>& Columns = TreeViewHeaderRow->GetColumns();
-	const int32 LastColumnIdx = Columns.Num() - 1;
 
 	// First column
-	if (Columns[0].ColumnId == ColumnId)
+	if (Columns.Num() == 0 ||
+		Columns[0].ColumnId == ColumnId)
 	{
 		return HAlign_Left;
 	}
 	// Last column
-	else if (Columns[LastColumnIdx].ColumnId == ColumnId)
+	else if (Columns[Columns.Num() - 1].ColumnId == ColumnId)
 	{
 		return HAlign_Right;
 	}
@@ -1762,13 +1761,6 @@ FText STableTreeView::TableRow_GetHighlightText() const
 	return SearchBox->GetText();
 }
 
-////////////////////////////////////////////////////////////////////////////////////////////////////
-
-FName STableTreeView::TableRow_GetHighlightedNodeName() const
-{
-	return HighlightedNodeName;
-}
-
 ////////////////////////////////////////////////////////////////////////////////////////////////////
 // Grouping
 ////////////////////////////////////////////////////////////////////////////////////////////////////
diff --git a/Engine/Source/Developer/TraceInsightsCore/Private/InsightsCore/Table/Widgets/STableTreeViewRow.cpp b/Engine/Source/Developer/TraceInsightsCore/Private/InsightsCore/Table/Widgets/STableTreeViewRow.cpp
--- a/Engine/Source/Developer/TraceInsightsCore/Private/InsightsCore/Table/Widgets/STableTreeViewRow.cpp
+++ b/Engine/Source/Developer/TraceInsightsCore/Private/InsightsCore/Table/Widgets/STableTreeViewRow.cpp
@@ -28,7 +28,6 @@ void STableTreeViewRow::Construct(const FArguments& InArgs, const TSharedRef<STa
 	SetHoveredCellDelegate = InArgs._OnSetHoveredCell;
 
 	HighlightText = InArgs._HighlightText;
-	HighlightedNodeName = InArgs._HighlightedNodeName;
 
 	TablePtr = InArgs._TablePtr;
 	TableTreeNodePtr = InArgs._TableTreeNodePtr;
@@ -131,10 +130,7 @@ FSlateColor STableTreeViewRow::GetBackgroundColorAndOpacity(double Time) const
 
 FSlateColor STableTreeViewRow::GetOutlineColorAndOpacity() const
 {
-	const FLinearColor NoColor(0.0f, 0.0f, 0.0f, 0.0f);
-	const bool bShouldBeHighlighted = TableTreeNodePtr->GetName() == HighlightedNodeName.Get();
-	const FLinearColor OutlineColorAndOpacity = bShouldBeHighlighted ? FLinearColor(FColorList::SlateBlue) : NoColor;
-	return OutlineColorAndOpacity;
+	return FLinearColor(0.0f, 0.0f, 0.0f, 0.0f);
 }
 
 ////////////////////////////////////////////////////////////////////////////////////////////////////
@@ -142,7 +138,7 @@ FSlateColor STableTreeViewRow::GetOutlineColorAndOpacity() const
 const FSlateBrush* STableTreeViewRow::GetOutlineBrush(const FName ColumnId) const
 {
 	EHorizontalAlignment HAlign = HAlign_Center;
-	if (IsColumnVisibleDelegate.IsBound())
+	if (GetColumnOutlineHAlignmentDelegate.IsBound())
 	{
 		HAlign = GetColumnOutlineHAlignmentDelegate.Execute(ColumnId);
 	}
@@ -153,21 +149,15 @@ const FSlateBrush* STableTreeViewRow::GetOutlineBrush(const FName ColumnId) cons
 
 bool STableTreeViewRow::HandleShouldBeEnabled() const
 {
-	bool bResult = false;
-
 	if (TableTreeNodePtr->IsGroup())
 	{
-		bResult = true;
+		return true;
 	}
-	else
+	else if (OnShouldBeEnabled.IsBound())
 	{
-		if (OnShouldBeEnabled.IsBound())
-		{
-			bResult = OnShouldBeEnabled.Execute(TableTreeNodePtr);
-		}
+		return OnShouldBeEnabled.Execute(TableTreeNodePtr);
 	}
-
-	return bResult;
+	return false;
 }
 
 ////////////////////////////////////////////////////////////////////////////////////////////////////
@@ -175,12 +165,10 @@ bool STableTreeViewRow::HandleShouldBeEnabled() const
 EVisibility STableTreeViewRow::IsColumnVisible(const FName ColumnId) const
 {
 	EVisibility Result = EVisibility::Collapsed;
-
 	if (IsColumnVisibleDelegate.IsBound())
 	{
 		Result = IsColumnVisibleDelegate.Execute(ColumnId) ? EVisibility::Visible : EVisibility::Collapsed;
 	}
-
 	return Result;
 }
 
diff --git a/Engine/Source/Developer/TraceInsightsCore/Public/InsightsCore/Table/Widgets/STableTreeView.h b/Engine/Source/Developer/TraceInsightsCore/Public/InsightsCore/Table/Widgets/STableTreeView.h
--- a/Engine/Source/Developer/TraceInsightsCore/Public/InsightsCore/Table/Widgets/STableTreeView.h
+++ b/Engine/Source/Developer/TraceInsightsCore/Public/InsightsCore/Table/Widgets/STableTreeView.h
@@ -275,7 +275,6 @@ protected:
 	UE_API EHorizontalAlignment TableRow_GetColumnOutlineHAlignment(const FName ColumnId) const;
 
 	UE_API FText TableRow_GetHighlightText() const;
-	UE_API FName TableRow_GetHighlightedNodeName() const;
 
 	////////////////////////////////////////////////////////////////////////////////////////////////////
 	// Node Filtering (TableRowNodes --> FilteredNodes)
@@ -537,9 +536,6 @@ protected:
 	/** A shared pointer to the tree node currently being hovered by the mouse. */
 	FTableTreeNodePtr HoveredNodePtr;
 
-	/** Name of the tree node that should be drawn as highlighted. */
-	FName HighlightedNodeName;
-
 	//////////////////////////////////////////////////
 	// Tree Nodes
 
diff --git a/Engine/Source/Developer/TraceInsightsCore/Public/InsightsCore/Table/Widgets/STableTreeViewRow.h b/Engine/Source/Developer/TraceInsightsCore/Public/InsightsCore/Table/Widgets/STableTreeViewRow.h
--- a/Engine/Source/Developer/TraceInsightsCore/Public/InsightsCore/Table/Widgets/STableTreeViewRow.h
+++ b/Engine/Source/Developer/TraceInsightsCore/Public/InsightsCore/Table/Widgets/STableTreeViewRow.h
@@ -38,7 +38,6 @@ public:
 		SLATE_EVENT(FGetColumnOutlineHAlignmentDelegate, OnGetColumnOutlineHAlignmentDelegate)
 		SLATE_EVENT(FSetHoveredTableTreeViewCell, OnSetHoveredCell)
 		SLATE_ATTRIBUTE(FText, HighlightText)
-		SLATE_ATTRIBUTE(FName, HighlightedNodeName)
 		SLATE_ARGUMENT(TSharedPtr<FTable>, TablePtr)
 		SLATE_ARGUMENT(FTableTreeNodePtr, TableTreeNodePtr)
 	SLATE_END_ARGS()
@@ -92,9 +91,6 @@ protected:
 	/** Text to be highlighted on tree node's name. */
 	TAttribute<FText> HighlightText;
 
-	/** Name of the tree node that should be drawn as highlighted. */
-	TAttribute<FName> HighlightedNodeName;
-
 	TSharedPtr<STableTreeRowToolTip> RowToolTip;
 };
 
