[Insights] MemAllocs: Improved grouping by Callstack (no more FName usage; uses less memory; ~2x speed improvement).

#Synced-CL 51291823

diff --git a/Engine/Source/Developer/TraceInsights/Private/Insights/MemoryProfiler/ViewModels/MemAllocGroupingByCallstack.cpp b/Engine/Source/Developer/TraceInsights/Private/Insights/MemoryProfiler/ViewModels/MemAllocGroupingByCallstack.cpp
--- a/Engine/Source/Developer/TraceInsights/Private/Insights/MemoryProfiler/ViewModels/MemAllocGroupingByCallstack.cpp
+++ b/Engine/Source/Developer/TraceInsights/Private/Insights/MemoryProfiler/ViewModels/MemAllocGroupingByCallstack.cpp
@@ -26,6 +26,15 @@ INSIGHTS_IMPLEMENT_RTTI(FCallstackFrameGroupNode)
 
 ////////////////////////////////////////////////////////////////////////////////////////////////////
 
+FCallstackFrameGroupNode::FCallstackFrameGroupNode(TWeakPtr<FTable> InParentTable, const TraceServices::FStackFrame* InStackFrame)
+	: FTableTreeNode(InParentTable)
+	, DisplayName(FMemAllocGroupingByCallstack::GetDisplayName(InStackFrame))
+	, StackFrame(InStackFrame)
+{
+}
+
+////////////////////////////////////////////////////////////////////////////////////////////////////
+
 const FText FCallstackFrameGroupNode::GetTooltipText() const
 {
 	if (StackFrame)
@@ -209,44 +218,34 @@ void FMemAllocGroupingByCallstack::GroupNodesInternal(
 							continue;
 						}
 
+						uint64 GroupNameId = 0;
 						if (bLocalIsGroupingByFunction)
 						{
-							const FName GroupName = GetGroupName(Frame);
+							// Merge groups by function name.
+							GroupNameId = GetGroupNameId(Frame);
 
-							// Merge with parent group, if it has the same name (i.e. same function).
-							if (GroupPtr->Parent != nullptr && GroupPtr->Name == GroupName)
+							// Merge with the parent group, if it has the same name id (i.e. same function name).
+							if (GroupPtr->Parent != nullptr && GroupPtr->NameId == GroupNameId)
 							{
 								GroupPtr = GroupPtr->Parent;
 							}
-
-							// Merge groups by name.
-							FCallstackGroup** GroupPtrPtr = GroupPtr->GroupMapByName.Find(GroupName);
-							if (!GroupPtrPtr)
-							{
-								GroupPtr = CreateGroup(CallstackGroups, GroupPtr, GroupName, InParentTable, Frame);
-								check(GroupPtr->Parent != nullptr);
-								GroupPtr->Parent->GroupMapByName.Add(GroupName, GroupPtr);
-							}
-							else
-							{
-								GroupPtr = *GroupPtrPtr;
-							}
 						}
 						else
 						{
 							// Merge groups by unique callstack frame.
-							FCallstackGroup** GroupPtrPtr = GroupPtr->GroupMap.Find(Frame->Addr);
-							if (!GroupPtrPtr)
-							{
-								const FName GroupName = GetGroupName(Frame);
-								GroupPtr = CreateGroup(CallstackGroups, GroupPtr, GroupName, InParentTable, Frame);
-								check(GroupPtr->Parent != nullptr);
-								GroupPtr->Parent->GroupMap.Add(Frame->Addr, GroupPtr);
-							}
-							else
-							{
-								GroupPtr = *GroupPtrPtr;
-							}
+							GroupNameId = Frame->Addr;
+						}
+
+						FCallstackGroup** FoundGroup = GroupPtr->GroupMap.Find(GroupNameId);
+						if (!FoundGroup)
+						{
+							GroupPtr = CreateGroup(CallstackGroups, GroupPtr, GroupNameId, InParentTable, Frame);
+							check(GroupPtr->Parent != nullptr);
+							GroupPtr->Parent->GroupMap.Add(GroupNameId, GroupPtr);
+						}
+						else
+						{
+							GroupPtr = *FoundGroup;
 						}
 					}
 
@@ -291,25 +290,38 @@ void FMemAllocGroupingByCallstack::GroupNodesInternal(
 
 ////////////////////////////////////////////////////////////////////////////////////////////////////
 
-FName FMemAllocGroupingByCallstack::GetGroupName(const TraceServices::FStackFrame* Frame) const
+uint64 FMemAllocGroupingByCallstack::GetGroupNameId(const TraceServices::FStackFrame* Frame) const
 {
 	const TraceServices::ESymbolQueryResult Result = Frame->Symbol->GetResult();
 	if (Result == TraceServices::ESymbolQueryResult::OK)
 	{
-		FStringView GroupName(Frame->Symbol->Name);
-		if (GroupName.Len() >= NAME_SIZE)
-		{
-			GroupName = FStringView(Frame->Symbol->Name, NAME_SIZE - 1);
-		}
-		return FName(GroupName, 0);
+		return uint64(Frame->Symbol->Name);
 	}
-	else if (Result == TraceServices::ESymbolQueryResult::Pending)
+	else
 	{
-		return FName(FString::Printf(TEXT("%s!0x%llX [...]"), Frame->Symbol->Module, Frame->Addr), 0);
+		return Frame->Addr;
+	}
+}
+
+////////////////////////////////////////////////////////////////////////////////////////////////////
+
+FText FMemAllocGroupingByCallstack::GetDisplayName(const TraceServices::FStackFrame* Frame)
+{
+	const TraceServices::ESymbolQueryResult Result = Frame->Symbol->GetResult();
+	if (Result == TraceServices::ESymbolQueryResult::OK)
+	{
+		return FText::FromString(Frame->Symbol->Name);
 	}
 	else
 	{
-		return FName(FString::Printf(TEXT("%s!0x%llX"), Frame->Symbol->Module, Frame->Addr), 0);
+		TStringBuilder<1024> Str;
+		Str.Append(Frame->Symbol->Module);
+		Str.Appendf(TEXT("!0x%llX"), Frame->Addr);
+		if (Result == TraceServices::ESymbolQueryResult::Pending)
+		{
+			Str.Append(TEXT(" [...]"));
+		}
+		return FText::FromStringView(Str.ToView());
 	}
 }
 
@@ -318,21 +330,21 @@ FName FMemAllocGroupingByCallstack::GetGroupName(const TraceServices::FStackFram
 FMemAllocGroupingByCallstack::FCallstackGroup* FMemAllocGroupingByCallstack::CreateGroup(
 	TArray<FCallstackGroup*>& InOutAllCallstackGroup,
 	FCallstackGroup* InParentGroup,
-	const FName InGroupName,
+	uint64 InGroupNameId,
 	TWeakPtr<FTable> InParentTable,
 	const TraceServices::FStackFrame* InFrame) const
 {
 	FCallstackGroup* NewGroupPtr = new FCallstackGroup();
 	NewGroupPtr->Parent = InParentGroup;
-	NewGroupPtr->Name = InGroupName;
+	NewGroupPtr->NameId = InGroupNameId;
 
 	InOutAllCallstackGroup.Add(NewGroupPtr);
 
-	TSharedPtr<FCallstackFrameGroupNode> NodePtr = MakeShared<FCallstackFrameGroupNode>(InGroupName, InParentTable, InFrame);
-	NodePtr->SetExpansion(false);
+	auto Node = MakeShared<FCallstackFrameGroupNode>(InParentTable, InFrame);
+	Node->SetExpansion(false);
 
-	InParentGroup->Node->AddChildAndSetParent(NodePtr);
-	NewGroupPtr->Node = NodePtr.Get();
+	InParentGroup->Node->AddChildAndSetParent(Node);
+	NewGroupPtr->Node = &Node.Get();
 
 	return NewGroupPtr;
 }
@@ -341,33 +353,33 @@ FMemAllocGroupingByCallstack::FCallstackGroup* FMemAllocGroupingByCallstack::Cre
 
 FTableTreeNode* FMemAllocGroupingByCallstack::CreateUnsetGroup(TWeakPtr<FTable> ParentTable, FTableTreeNode& Parent) const
 {
-	static FName GroupName(GetCallstackNotAvailableString().ToString());
-	FTableTreeNodePtr NodePtr = MakeShared<FTableTreeNode>(GroupName, ParentTable);
-	NodePtr->SetExpansion(false);
-	Parent.AddChildAndSetParent(NodePtr);
-	return NodePtr.Get();
+	const FText DisplayName = GetCallstackNotAvailableString();
+	auto GroupNode = MakeShared<FCustomTableTreeNode>(ParentTable, DisplayName);
+	GroupNode->SetExpansion(false);
+	Parent.AddChildAndSetParent(GroupNode);
+	return &GroupNode.Get();
 }
 
 ////////////////////////////////////////////////////////////////////////////////////////////////////
 
 FTableTreeNode* FMemAllocGroupingByCallstack::CreateNoCallstackGroup(TWeakPtr<FTable> ParentTable, FTableTreeNode& Parent) const
 {
-	static FName GroupName(GetNoCallstackString().ToString());
-	FTableTreeNodePtr NodePtr = MakeShared<FTableTreeNode>(GroupName, ParentTable);
-	NodePtr->SetExpansion(false);
-	Parent.AddChildAndSetParent(NodePtr);
-	return NodePtr.Get();
+	const FText DisplayName = GetNoCallstackString();
+	auto GroupNode = MakeShared<FCustomTableTreeNode>(ParentTable, DisplayName);
+	GroupNode->SetExpansion(false);
+	Parent.AddChildAndSetParent(GroupNode);
+	return &GroupNode.Get();
 }
 
 ////////////////////////////////////////////////////////////////////////////////////////////////////
 
 FTableTreeNode* FMemAllocGroupingByCallstack::CreateEmptyCallstackGroup(TWeakPtr<FTable> ParentTable, FTableTreeNode& Parent) const
 {
-	static FName GroupName(GetEmptyCallstackString().ToString());
-	FTableTreeNodePtr NodePtr = MakeShared<FTableTreeNode>(GroupName, ParentTable);
-	NodePtr->SetExpansion(false);
-	Parent.AddChildAndSetParent(NodePtr);
-	return NodePtr.Get();
+	const FText DisplayName = GetEmptyCallstackString();
+	auto GroupNode = MakeShared<FCustomTableTreeNode>(ParentTable, DisplayName);
+	GroupNode->SetExpansion(false);
+	Parent.AddChildAndSetParent(GroupNode);
+	return &GroupNode.Get();
 }
 
 ////////////////////////////////////////////////////////////////////////////////////////////////////
diff --git a/Engine/Source/Developer/TraceInsights/Private/Insights/MemoryProfiler/ViewModels/MemAllocGroupingByCallstack.h b/Engine/Source/Developer/TraceInsights/Private/Insights/MemoryProfiler/ViewModels/MemAllocGroupingByCallstack.h
--- a/Engine/Source/Developer/TraceInsights/Private/Insights/MemoryProfiler/ViewModels/MemAllocGroupingByCallstack.h
+++ b/Engine/Source/Developer/TraceInsights/Private/Insights/MemoryProfiler/ViewModels/MemAllocGroupingByCallstack.h
@@ -22,14 +22,15 @@ class FCallstackFrameGroupNode : public FTableTreeNode
 
 public:
 	/** Initialization constructor for the group node. */
-	explicit FCallstackFrameGroupNode(const FName InName, TWeakPtr<FTable> InParentTable, const TraceServices::FStackFrame* InStackFrame)
-		: FTableTreeNode(InName, InParentTable)
-		, StackFrame(InStackFrame)
-	{
-	}
+	explicit FCallstackFrameGroupNode(TWeakPtr<FTable> InParentTable, const TraceServices::FStackFrame* InStackFrame);
 
 	virtual ~FCallstackFrameGroupNode() = default;
 
+	virtual const FText GetDisplayName() const override
+	{
+		return DisplayName;
+	}
+
 	virtual const FText GetTooltipText() const override;
 	virtual const FSlateBrush* GetIcon() const override;
 	virtual FLinearColor GetIconColor() const override;
@@ -44,6 +45,7 @@ public:
 	}
 
 private:
+	FText DisplayName;
 	const TraceServices::FStackFrame* StackFrame = nullptr;
 };
 
@@ -57,11 +59,10 @@ private:
 	struct FCallstackGroup
 	{
 		FCallstackGroup* Parent = nullptr;
-		FName Name;
+		uint64 NameId = 0;
 		FTableTreeNode* Node = nullptr;
 		const TraceServices::FStackFrame* Frame = nullptr;
-		TMap<uint64, FCallstackGroup*> GroupMap; // Callstack Frame Address -> FCallstackGroup*
-		TMap<FName, FCallstackGroup*> GroupMapByName; // Group Name --> FCallstackGroup*
+		TMap<uint64, FCallstackGroup*> GroupMap; // Callstack Frame Address or Group Name Id -> FCallstackGroup*
 	};
 
 public:
@@ -78,13 +79,17 @@ public:
 
 	bool ShouldSkipFilteredFrames() const { return bShouldSkipFilteredFrames; }
 	void SetSkipFilteredFrames(bool bOnOff) { bShouldSkipFilteredFrames = bOnOff; }
+
+	static FText GetDisplayName(const TraceServices::FStackFrame* Frame);
+
 protected:
 	FMemAllocGroupingByCallstack(const FText& InShortName, const FText& InTitleName, const FText& InDescription, const FSlateBrush* InIcon, bool bInIsAllocCallstack, bool bInIsInverted, bool bInIsGroupingByFunction);
 	void GroupNodesInternal(const TArray<FTableTreeNodePtr>& Nodes, FTableTreeNode& ParentGroup, TWeakPtr<FTable> InParentTable, IAsyncOperationProgress& InAsyncOperationProgress, TUniqueFunction<FTableTreeNodePtr(FTableTreeNodePtr&)> ProcessNode) const;
+
 private:
-	FName GetGroupName(const TraceServices::FStackFrame* Frame) const;
+	uint64 GetGroupNameId(const TraceServices::FStackFrame* Frame) const;
 
-	FCallstackGroup* CreateGroup(TArray<FCallstackGroup*>& InOutAllCallstackGroup, FCallstackGroup* InParentGroup, const FName InGroupName, TWeakPtr<FTable> InParentTable, const TraceServices::FStackFrame* InFrame) const;
+	FCallstackGroup* CreateGroup(TArray<FCallstackGroup*>& InOutAllCallstackGroup, FCallstackGroup* InParentGroup, uint64 InGroupNameId, TWeakPtr<FTable> InParentTable, const TraceServices::FStackFrame* InFrame) const;
 
 	FTableTreeNode* CreateUnsetGroup(TWeakPtr<FTable> ParentTable, FTableTreeNode& Parent) const;
 	FTableTreeNode* CreateNoCallstackGroup(TWeakPtr<FTable> ParentTable, FTableTreeNode& Parent) const;
