[Insights]
* MemAllocs: Improved grouping by Tags (removed the deprecated FName usage in FMemTagTableTreeNode; no more FNames during grouping; faster grouping).
* MemAllocs: Added the full tag name in the tooltip of Tag group nodes.

#Synced-CL 51291823

diff --git a/Engine/Source/Developer/TraceInsights/Private/Insights/MemoryProfiler/ViewModels/MemAllocGroupingByTag.cpp b/Engine/Source/Developer/TraceInsights/Private/Insights/MemoryProfiler/ViewModels/MemAllocGroupingByTag.cpp
--- a/Engine/Source/Developer/TraceInsights/Private/Insights/MemoryProfiler/ViewModels/MemAllocGroupingByTag.cpp
+++ b/Engine/Source/Developer/TraceInsights/Private/Insights/MemoryProfiler/ViewModels/MemAllocGroupingByTag.cpp
@@ -55,7 +55,7 @@ void FMemAllocGroupingByTag::GroupNodes(
 	constexpr TagIdType InvalidTagId = ~0u;
 	struct FTagNode
 	{
-		FTagNode(FName InName, TagIdType InId, TagIdType InParentId, FTableTreeNodePtr InTreeNode)
+		FTagNode(const TCHAR* InName, TagIdType InId, TagIdType InParentId, FTableTreeNodePtr InTreeNode)
 			: Name(InName), Id(InId), ParentId(InParentId), TreeNode(InTreeNode), Parent(nullptr), AllocCount(0)
 		{}
 
@@ -68,6 +68,9 @@ void FMemAllocGroupingByTag::GroupNodes(
 				for (int32 TargetIndex = 0; TargetIndex < Index; ++TargetIndex)
 				{
 					FTagNode* TargetNode = Children[TargetIndex];
+
+					//if (FCString::Strcmp(TargetNode->Name, CurrentNode->Name) == 0)
+					// Name is persistent string so we can compare pointers directly.
 					if (TargetNode->Name == CurrentNode->Name)
 					{
 						// Move children from CurrentNode to TargetNode.
@@ -91,7 +94,7 @@ void FMemAllocGroupingByTag::GroupNodes(
 			}
 		};
 
-		FName Name;
+		const TCHAR* Name; // persistent string (from Session's StringStore or static string)
 
 		TagIdType Id;
 		TagIdType ParentId;
@@ -129,14 +132,8 @@ void FMemAllocGroupingByTag::GroupNodes(
 		TraceServices::FProviderReadScopeLock TagProviderReadScopeLock(TagProvider);
 		TagProvider.EnumerateTags([&](const TCHAR* Name, const TCHAR* FullName, TagIdType Id, TagIdType ParentId)
 			{
-				const FName NodeName(Name);
-#if INSIGHTS_MERGE_MEM_TAGS_BY_NAME
-				FTableTreeNodePtr TreeNode = MakeShared<FMemTagTableTreeNode>(NodeName, InParentTable, FullName);
-#else
-				const FName NodeNameEx(Name, int32(Id));
-				FTableTreeNodePtr TreeNode = MakeShared<FMemTagTableTreeNode>(NodeNameEx, InParentTable, FullName);
-#endif
-				FTagNode* TagNode = new FTagNode(NodeName, Id, ParentId, TreeNode);
+				FTableTreeNodePtr TreeNode = MakeShared<FMemTagTableTreeNode>(InParentTable, FullName, Name);
+				FTagNode* TagNode = new FTagNode(Name, Id, ParentId, TreeNode);
 				TagNodes.Add(TagNode);
 				IdToNodeMap.Add(Id, TagNode);
 			});
@@ -147,9 +144,9 @@ void FMemAllocGroupingByTag::GroupNodes(
 	FTagNode* UntaggedNode = IdToNodeMap.FindRef(UntaggedNodeId);
 	if (!UntaggedNode)
 	{
-		const FName NodeName(TEXT("Untagged"));
-		FTableTreeNodePtr UntaggedTreeNode = MakeShared<FMemTagTableTreeNode>(NodeName, InParentTable, TEXT("Untagged"));
-		UntaggedNode = new FTagNode(NodeName, UntaggedNodeId, InvalidTagId, UntaggedTreeNode);
+		const TCHAR* UntaggedNodeName = TEXT("Untagged");
+		FTableTreeNodePtr UntaggedTreeNode = MakeShared<FMemTagTableTreeNode>(InParentTable, UntaggedNodeName, UntaggedNodeName);
+		UntaggedNode = new FTagNode(UntaggedNodeName, UntaggedNodeId, InvalidTagId, UntaggedTreeNode);
 		TagNodes.Add(UntaggedNode);
 		IdToNodeMap.Add(UntaggedNodeId, UntaggedNode);
 	}
@@ -159,15 +156,16 @@ void FMemAllocGroupingByTag::GroupNodes(
 	FTagNode* UnknownTagNode = IdToNodeMap.FindRef(UnknownTagNodeId);
 	if (!UnknownTagNode)
 	{
-		const FName NodeName(TEXT("Unknown"));
-		FTableTreeNodePtr UnknownTagTreeNode = MakeShared<FMemTagTableTreeNode>(NodeName, InParentTable, TEXT("Unknown"));
-		UnknownTagNode = new FTagNode(NodeName, UnknownTagNodeId, InvalidTagId, UnknownTagTreeNode);
+		const TCHAR* UnknownNodeName = TEXT("Unknown");
+		FTableTreeNodePtr UnknownTagTreeNode = MakeShared<FMemTagTableTreeNode>(InParentTable, UnknownNodeName, UnknownNodeName);
+		UnknownTagNode = new FTagNode(UnknownNodeName, UnknownTagNodeId, InvalidTagId, UnknownTagTreeNode);
 		TagNodes.Add(UnknownTagNode);
 		IdToNodeMap.Add(UnknownTagNodeId, UnknownTagNode);
 	}
 
 	// Create the virtual tree of tag group nodes.
-	FTagNode Root(FName(TEXT("Root")), InvalidTagId, InvalidTagId, nullptr);
+	const TCHAR* RootNodeName = TEXT("Root");
+	FTagNode Root(RootNodeName, InvalidTagId, InvalidTagId, nullptr);
 	for (FTagNode* TagNode : TagNodes)
 	{
 		if (TagNode->ParentId == InvalidTagId)
diff --git a/Engine/Source/Developer/TraceInsights/Private/Insights/MemoryProfiler/ViewModels/MemAllocGroupingByTag.h b/Engine/Source/Developer/TraceInsights/Private/Insights/MemoryProfiler/ViewModels/MemAllocGroupingByTag.h
--- a/Engine/Source/Developer/TraceInsights/Private/Insights/MemoryProfiler/ViewModels/MemAllocGroupingByTag.h
+++ b/Engine/Source/Developer/TraceInsights/Private/Insights/MemoryProfiler/ViewModels/MemAllocGroupingByTag.h
@@ -39,21 +39,33 @@ class FMemTagTableTreeNode : public FTableTreeNode
 
 public:
 	/** Initialization constructor for a table record node. */
-	explicit FMemTagTableTreeNode(const FName InName, TWeakPtr<FTable> InParentTable, int32 InRowIndex, const TCHAR* InTagFullName, bool IsGroup)
-		: FTableTreeNode(InName, InParentTable, InRowIndex, IsGroup)
+	explicit FMemTagTableTreeNode(TWeakPtr<FTable> InParentTable, int32 InRowIndex, const TCHAR* InTagFullName, const TCHAR* InName, bool IsGroup)
+		: FTableTreeNode(InParentTable, InRowIndex, IsGroup)
+		, DisplayName(FText::FromString(InName))
 		, TagFullName(InTagFullName)
 	{
 	}
 
 	/** Initialization constructor for the group node. */
-	explicit FMemTagTableTreeNode(const FName InName, TWeakPtr<FTable> InParentTable, const TCHAR* InTagFullName)
-		: FTableTreeNode(InName, InParentTable)
+	explicit FMemTagTableTreeNode(TWeakPtr<FTable> InParentTable, const TCHAR* InTagFullName, const TCHAR* InName)
+		: FTableTreeNode(InParentTable)
+		, DisplayName(FText::FromString(InName))
 		, TagFullName(InTagFullName)
 	{
 	}
 
 	virtual ~FMemTagTableTreeNode() = default;
 
+	virtual const FText GetDisplayName() const override
+	{
+		return DisplayName;
+	}
+
+	virtual const FText GetTooltipText() const
+	{
+		return FText::FromString(TagFullName);
+	}
+
 	virtual FLinearColor GetIconColor() const override
 	{
 		return FLinearColor(0.75f, 0.5f, 1.0f, 1.0f);
@@ -77,6 +89,7 @@ private:
 	void UpdateLLMSize() const;
 
 private:
+	FText DisplayName;
 	const TCHAR* TagFullName = nullptr;
 	mutable int64 LLMSize = INT64_MAX;
 };
