[Insights]
* Fixed deprecated FName usage in CookProfiler (FPackageNode, SPackageTableTreeView).

#rnx

#Synced-CL 51295847

diff --git a/Engine/Source/Developer/TraceInsights/Private/Insights/CookProfiler/CookProfilerManager.cpp b/Engine/Source/Developer/TraceInsights/Private/Insights/CookProfiler/CookProfilerManager.cpp
--- a/Engine/Source/Developer/TraceInsights/Private/Insights/CookProfiler/CookProfilerManager.cpp
+++ b/Engine/Source/Developer/TraceInsights/Private/Insights/CookProfiler/CookProfilerManager.cpp
@@ -128,7 +128,7 @@ void FCookProfilerManager::UnregisterMajorTabs()
 
 bool FCookProfilerManager::Tick(float DeltaTime)
 {
-	// Check if session has task events (to spawn the tab), but not too often.
+	// Check if session has Cook events (to spawn the Packages tab), but not too often.
 	if (!bIsAvailable && AvailabilityCheck.Tick())
 	{
 		bool bShouldBeAvailable = false;
@@ -137,7 +137,6 @@ bool FCookProfilerManager::Tick(float DeltaTime)
 		if (Session.IsValid())
 		{
 			const TraceServices::ICookProfilerProvider* CookProvider = TraceServices::ReadCookProfilerProvider(*Session.Get());
-
 			if (CookProvider)
 			{
 				TraceServices::FProviderReadScopeLock ProviderReadScope(*CookProvider);
diff --git a/Engine/Source/Developer/TraceInsights/Private/Insights/CookProfiler/ViewModels/PackageNode.h b/Engine/Source/Developer/TraceInsights/Private/Insights/CookProfiler/ViewModels/PackageNode.h
--- a/Engine/Source/Developer/TraceInsights/Private/Insights/CookProfiler/ViewModels/PackageNode.h
+++ b/Engine/Source/Developer/TraceInsights/Private/Insights/CookProfiler/ViewModels/PackageNode.h
@@ -22,17 +22,17 @@ namespace UE::Insights::CookProfiler
 
 class FPackageNode;
 
-/** Type definition for shared pointers to instances of FTaskNode. */
+/** Type definition for shared pointers to instances of FPackageNode. */
 typedef TSharedPtr<class FPackageNode> FPackageNodePtr;
 
-/** Type definition for shared references to instances of FTaskNode. */
+/** Type definition for shared references to instances of FPackageNode. */
 typedef TSharedRef<class FPackageNode> FPackageNodeRef;
 
-/** Type definition for shared references to const instances of FTaskNode. */
+/** Type definition for shared references to const instances of FPackageNode. */
 typedef TSharedRef<const class FPackageNode> FPackageNodeRefConst;
 
-/** Type definition for weak references to instances of FTaskNode. */
-typedef TWeakPtr<class FTaskNode> FPackageNodeWeak;
+/** Type definition for weak references to instances of FPackageNode. */
+typedef TWeakPtr<class FPackageNode> FPackageNodeWeak;
 
 ////////////////////////////////////////////////////////////////////////////////////////////////////
 /**
@@ -43,16 +43,18 @@ class FPackageNode : public FTableTreeNode
 	INSIGHTS_DECLARE_RTTI(FPackageNode, FTableTreeNode)
 
 public:
-	/** Initialization constructor for the Task node. */
-	explicit FPackageNode(const FName InName, TWeakPtr<FPackageTable> InParentTable, int32 InRowIndex)
-		: FTableTreeNode(InName, InParentTable, InRowIndex)
+	/** Initialization constructor for the Package node.
+	 * InPackageName needs to be a persistent string (ex: from the session's string store).
+	 */
+	explicit FPackageNode(TWeakPtr<FPackageTable> InParentTable, int32 InRowIndex, const TCHAR* InPackageName)
+		: FTableTreeNode(InParentTable, InRowIndex)
+		, PackageName(InPackageName)
 	{
 	}
 
-	/** Initialization constructor for the group node. */
-	explicit FPackageNode(const FName InGroupName, TWeakPtr<FPackageTable> InParentTable)
-		: FTableTreeNode(InGroupName, InParentTable)
+	virtual const FText GetDisplayName() const override
 	{
+		return FText::FromString(PackageName);
 	}
 
 	FPackageTable& GetPackageTableChecked() const
@@ -65,6 +67,9 @@ public:
 	bool IsValidPackage() const { return GetPackageTableChecked().IsValidRowIndex(GetRowIndex()); }
 	const FPackageEntry* GetPackage() const { return GetPackageTableChecked().GetPackage(GetRowIndex()); }
 	const FPackageEntry& GetPackageChecked() const { return GetPackageTableChecked().GetPackageChecked(GetRowIndex()); }
+
+private:
+	const TCHAR* PackageName = nullptr; // persistent string (from the session's string store)
 };
 
 ////////////////////////////////////////////////////////////////////////////////////////////////////
diff --git a/Engine/Source/Developer/TraceInsights/Private/Insights/CookProfiler/ViewModels/PackageTable.cpp b/Engine/Source/Developer/TraceInsights/Private/Insights/CookProfiler/ViewModels/PackageTable.cpp
--- a/Engine/Source/Developer/TraceInsights/Private/Insights/CookProfiler/ViewModels/PackageTable.cpp
+++ b/Engine/Source/Developer/TraceInsights/Private/Insights/CookProfiler/ViewModels/PackageTable.cpp
@@ -27,7 +27,7 @@ const FName FPackageTableColumns::LoadTimeInclColumnId(TEXT("LoadTimeIncl"));
 const FName FPackageTableColumns::LoadTimeExclColumnId(TEXT("LoadTimeExcl"));
 
 const FName FPackageTableColumns::SaveTimeInclColumnId(TEXT("SaveTimeIncl"));
-const FName FPackageTableColumns::SaveTimeExclColumnId(TEXT("SaveTimeExcl")); 
+const FName FPackageTableColumns::SaveTimeExclColumnId(TEXT("SaveTimeExcl"));
 
 const FName FPackageTableColumns::BeginCacheForCookedPlatformDataTimeInclColumnId(TEXT("BeginCacheForCookedPlatformDataTimeIncl"));
 const FName FPackageTableColumns::BeginCacheForCookedPlatformDataTimeExclColumnId(TEXT("BeginCacheForCookedPlatformDataTimeExcl"));
@@ -92,7 +92,7 @@ struct DefaultPackageFieldGetterFuncts
 };
 
 ////////////////////////////////////////////////////////////////////////////////////////////////////
-// FTaskTable
+// FPackageTable
 ////////////////////////////////////////////////////////////////////////////////////////////////////
 
 FPackageTable::FPackageTable()
diff --git a/Engine/Source/Developer/TraceInsights/Private/Insights/CookProfiler/Widgets/SPackageTableTreeView.cpp b/Engine/Source/Developer/TraceInsights/Private/Insights/CookProfiler/Widgets/SPackageTableTreeView.cpp
--- a/Engine/Source/Developer/TraceInsights/Private/Insights/CookProfiler/Widgets/SPackageTableTreeView.cpp
+++ b/Engine/Source/Developer/TraceInsights/Private/Insights/CookProfiler/Widgets/SPackageTableTreeView.cpp
@@ -121,6 +121,11 @@ void SPackageTableTreeView::RebuildTree(bool bResync)
 		return;
 	}
 
+	if (!Session.IsValid())
+	{
+		return;
+	}
+
 	TraceServices::FAnalysisSessionReadScope SessionReadScope(*Session.Get());
 
 	if (!Session->IsAnalysisComplete())
@@ -150,10 +155,9 @@ void SPackageTableTreeView::RebuildTree(bool bResync)
 		{
 			Packages.Emplace(Package);
 			uint32 Index = static_cast<uint32>(Packages.Num() - 1);
-			FName NodeName(Package.Name);
-			FPackageNodePtr NodePtr = MakeShared<FPackageNode>(NodeName, PackageTable, Index);
+			FPackageNodePtr NodePtr = MakeShared<FPackageNode>(PackageTable, Index, Package.Name);
 			Nodes->Add(NodePtr);
-		};
+		}
 	}
 
 	bDataLoaded = true;
