[Insights]
* Fixed deprecated FName usage in TaskGraphProfiler.

#rnx

#Synced-CL 51295847

diff --git a/Engine/Source/Developer/TraceInsights/Private/Insights/TaskGraphProfiler/ViewModels/TaskNode.cpp b/Engine/Source/Developer/TraceInsights/Private/Insights/TaskGraphProfiler/ViewModels/TaskNode.cpp
--- a/Engine/Source/Developer/TraceInsights/Private/Insights/TaskGraphProfiler/ViewModels/TaskNode.cpp
+++ b/Engine/Source/Developer/TraceInsights/Private/Insights/TaskGraphProfiler/ViewModels/TaskNode.cpp
@@ -2,9 +2,18 @@
 
 #include "TaskNode.h"
 
+#define LOCTEXT_NAMESPACE "UE::Insights::TaskGraphProfiler::FTaskNode"
+
 namespace UE::Insights::TaskGraphProfiler
 {
 
 INSIGHTS_IMPLEMENT_RTTI(FTaskNode)
 
+const FText FTaskNode::GetDisplayName() const
+{
+	return FText::Format(LOCTEXT("TaskNodeNameFmt", "task_{0}"), FText::AsNumber(GetRowIndex()));
+}
+
 } // namespace UE::Insights::TaskGraphProfiler
+
+#undef LOCTEXT_NAMESPACE
diff --git a/Engine/Source/Developer/TraceInsights/Private/Insights/TaskGraphProfiler/ViewModels/TaskNode.h b/Engine/Source/Developer/TraceInsights/Private/Insights/TaskGraphProfiler/ViewModels/TaskNode.h
--- a/Engine/Source/Developer/TraceInsights/Private/Insights/TaskGraphProfiler/ViewModels/TaskNode.h
+++ b/Engine/Source/Developer/TraceInsights/Private/Insights/TaskGraphProfiler/ViewModels/TaskNode.h
@@ -42,14 +42,8 @@ class FTaskNode : public FTableTreeNode
 
 public:
 	/** Initialization constructor for the Task node. */
-	explicit FTaskNode(const FName InName, TWeakPtr<FTaskTable> InParentTable, int32 InRowIndex)
-		: FTableTreeNode(InName, InParentTable, InRowIndex)
-	{
-	}
-
-	/** Initialization constructor for the group node. */
-	explicit FTaskNode(const FName InGroupName, TWeakPtr<FTaskTable> InParentTable)
-		: FTableTreeNode(InGroupName, InParentTable)
+	explicit FTaskNode(TWeakPtr<FTaskTable> InParentTable, int32 InRowIndex)
+		: FTableTreeNode(InParentTable, InRowIndex)
 	{
 	}
 
@@ -60,6 +54,8 @@ public:
 		return *StaticCastSharedPtr<FTaskTable>(TablePin);
 	}
 
+	virtual const FText GetDisplayName() const override;
+
 	bool IsValidTask() const { return GetTaskTableChecked().IsValidRowIndex(GetRowIndex()); }
 	const FTaskEntry* GetTask() const { return GetTaskTableChecked().GetTask(GetRowIndex()); }
 	const FTaskEntry& GetTaskChecked() const { return GetTaskTableChecked().GetTaskChecked(GetRowIndex()); }
diff --git a/Engine/Source/Developer/TraceInsights/Private/Insights/TaskGraphProfiler/Widgets/STaskTableTreeView.cpp b/Engine/Source/Developer/TraceInsights/Private/Insights/TaskGraphProfiler/Widgets/STaskTableTreeView.cpp
--- a/Engine/Source/Developer/TraceInsights/Private/Insights/TaskGraphProfiler/Widgets/STaskTableTreeView.cpp
+++ b/Engine/Source/Developer/TraceInsights/Private/Insights/TaskGraphProfiler/Widgets/STaskTableTreeView.cpp
@@ -223,25 +223,22 @@ void STaskTableTreeView::RebuildTree(bool bResync)
 
 		if (QueryStartTime < QueryEndTime && Session.IsValid())
 		{
-			TraceServices::FAnalysisSessionReadScope SessionReadScope(*Session.Get());
-
 			const TraceServices::ITasksProvider* TasksProvider = TraceServices::ReadTasksProvider(*Session.Get());
-
 			if (TasksProvider)
 			{
-				FName BaseNodeName(TEXT("task"));
+				TraceServices::FAnalysisSessionReadScope SessionReadScope(*Session.Get());
 
 				TArray<FTableTreeNodePtr>* Nodes = &TableRowNodes;
-
-				TasksProvider->EnumerateTasks(QueryStartTime, QueryEndTime, SelectedTasksSelectionOption, [&Tasks, &TaskTable, &BaseNodeName, Nodes](const TraceServices::FTaskInfo& TaskInfo)
-				{
-					Tasks.Emplace(TaskInfo);
-					uint32 Index = Tasks.Num() - 1;
-					FName NodeName(BaseNodeName, static_cast<int32>(TaskInfo.Id + 1));
-					FTaskNodePtr NodePtr = MakeShared<FTaskNode>(NodeName, TaskTable, Index);
-					Nodes->Add(NodePtr);
-					return TraceServices::ETaskEnumerationResult::Continue;
-				});
+				TasksProvider->EnumerateTasks(QueryStartTime, QueryEndTime, SelectedTasksSelectionOption,
+					[&Tasks, &TaskTable, Nodes]
+					(const TraceServices::FTaskInfo& TaskInfo)
+					{
+						Tasks.Emplace(TaskInfo);
+						int32 Index = Tasks.Num() - 1;
+						FTaskNodePtr NodePtr = MakeShared<FTaskNode>(TaskTable, Index);
+						Nodes->Add(NodePtr);
+						return TraceServices::ETaskEnumerationResult::Continue;
+					});
 			}
 		}
 
