[Insights] * Fixed CopyNameToClipboard for timers and counters to copy the DisplayName.

#rnx

#Synced-CL 51295847

diff --git a/Engine/Source/Developer/TraceInsights/Private/Insights/Table/ViewModels/TableCommands.cpp b/Engine/Source/Developer/TraceInsights/Private/Insights/Table/ViewModels/TableCommands.cpp
--- a/Engine/Source/Developer/TraceInsights/Private/Insights/Table/ViewModels/TableCommands.cpp
+++ b/Engine/Source/Developer/TraceInsights/Private/Insights/Table/ViewModels/TableCommands.cpp
@@ -8,7 +8,8 @@
 
 namespace UE::Insights::Private
 {
-void CopyToClipboardImpl(TSharedRef<FTable> Table, TArray<Insights::FBaseTreeNodePtr> SelectedNodes, TSharedPtr<Insights::ITableCellValueSorter> Sorter, ESortMode ColumnSortMode)
+
+void CopyToClipboardImpl(TSharedRef<FTable> Table, TArray<FBaseTreeNodePtr> SelectedNodes, TSharedPtr<ITableCellValueSorter> Sorter, ESortMode ColumnSortMode)
 {
 	if (SelectedNodes.Num() == 0)
 	{
@@ -30,7 +31,7 @@ void CopyToClipboardImpl(TSharedRef<FTable> Table, TArray<Insights::FBaseTreeNod
 	}
 }
 
-void CopyNameToClipboardImpl(TSharedRef<FTable> Table, TArray<Insights::FBaseTreeNodePtr> SelectedNodes, TSharedPtr<ITableCellValueSorter> Sorter, ESortMode ColumnSortMode)
+void CopyNameToClipboardImpl(TSharedRef<FTable> Table, TArray<FBaseTreeNodePtr> SelectedNodes, TSharedPtr<ITableCellValueSorter> Sorter, ESortMode ColumnSortMode)
 {
 	if (SelectedNodes.Num() == 0)
 	{
@@ -45,10 +46,10 @@ void CopyNameToClipboardImpl(TSharedRef<FTable> Table, TArray<Insights::FBaseTre
 	FString ClipboardText = FString::JoinBy(
 		SelectedNodes,
 		TEXT("\n"),
-		[](const Insights::FBaseTreeNodePtr& TimerNode)
-	{
-		return TimerNode->GetName().ToString();
-	});
+		[](const FBaseTreeNodePtr& Node)
+		{
+			return Node->GetDisplayName().ToString();
+		});
 
 	if (ClipboardText.Len() > 0)
 	{
@@ -56,4 +57,4 @@ void CopyNameToClipboardImpl(TSharedRef<FTable> Table, TArray<Insights::FBaseTre
 	}
 }
 
-} // namespace UE::Insights::Private
\ No newline at end of file
+} // namespace UE::Insights::Private
