[Insights] Fixed the deprecated FName usage in FTimerNode, STimersView and in STimerTreeView. Timers will not create FNames anymore.

#Synced-CL 51334642

diff --git a/Engine/Source/Developer/TraceInsights/Private/Insights/TimingProfiler/ViewModels/TimerGroupingAndSorting.cpp b/Engine/Source/Developer/TraceInsights/Private/Insights/TimingProfiler/ViewModels/TimerGroupingAndSorting.cpp
--- a/Engine/Source/Developer/TraceInsights/Private/Insights/TimingProfiler/ViewModels/TimerGroupingAndSorting.cpp
+++ b/Engine/Source/Developer/TraceInsights/Private/Insights/TimingProfiler/ViewModels/TimerGroupingAndSorting.cpp
@@ -12,7 +12,7 @@
 #define INSIGHTS_ENSURE(...)
 
 // Sort by name (ascending).
-#define INSIGHTS_DEFAULT_SORTING_NODES(A, B) return A->GetName().LexicalLess(B->GetName());
+#define INSIGHTS_DEFAULT_SORTING_NODES(A, B) return (A->GetDisplayName().CompareTo(B->GetDisplayName()) < 0);
 //#define INSIGHTS_DEFAULT_SORTING_NODES(A, B) return A->GetDefaultSortOrder() < B->GetDefaultSortOrder();
 
 #define LOCTEXT_NAMESPACE "UE::Insights::TimingProfiler::TimerNode"
diff --git a/Engine/Source/Developer/TraceInsights/Private/Insights/TimingProfiler/ViewModels/TimerNode.cpp b/Engine/Source/Developer/TraceInsights/Private/Insights/TimingProfiler/ViewModels/TimerNode.cpp
--- a/Engine/Source/Developer/TraceInsights/Private/Insights/TimingProfiler/ViewModels/TimerNode.cpp
+++ b/Engine/Source/Developer/TraceInsights/Private/Insights/TimingProfiler/ViewModels/TimerNode.cpp
@@ -63,7 +63,8 @@ FName MakeSafeFName(const TCHAR* InName)
 ////////////////////////////////////////////////////////////////////////////////////////////////////
 
 FTimerNode::FTimerNode(uint32 InTimerId, const TCHAR* InName, ETimerNodeType InType, bool bInIsGroup)
-	: FBaseTreeNode(MakeSafeFName(InName), bInIsGroup)
+	: FBaseTreeNode(bInIsGroup)
+	, DisplayName(FText::FromString(InName))
 	, TimerId(InTimerId)
 	, MetaGroupName(InType == ETimerNodeType::GpuScope ? GetGpuScopeGroup() :
 					InType == ETimerNodeType::VerseSampling ? GetVerseSamplingGroup() :
@@ -76,16 +77,16 @@ FTimerNode::FTimerNode(uint32 InTimerId, const TCHAR* InName, ETimerNodeType InT
 	Color.R = ((Color32 >> 16) & 0xFF) / 255.0f;
 	Color.G = ((Color32 >>  8) & 0xFF) / 255.0f;
 	Color.B = ((Color32      ) & 0xFF) / 255.0f;
-	Color.A = 1.0;
+	Color.A = 1.0f;
 
 	ResetAggregatedStats();
 }
 
 ////////////////////////////////////////////////////////////////////////////////////////////////////
 
-/** Initialization constructor for the group node. */
-FTimerNode::FTimerNode(const FName InGroupName)
-	: FBaseTreeNode(InGroupName, true)
+FTimerNode::FTimerNode(const FText& InDisplayName)
+	: FBaseTreeNode(true)
+	, DisplayName(InDisplayName)
 	, TimerId(InvalidTimerId)
 	, Type(ETimerNodeType::Group)
 	, Color(0.0f, 0.0f, 0.0f, 1.0f)
@@ -101,13 +102,6 @@ FTimerNode::~FTimerNode()
 
 ////////////////////////////////////////////////////////////////////////////////////////////////////
 
-void FTimerNode::SetColor(const FLinearColor& InColor)
-{
-	Color = InColor;
-}
-
-////////////////////////////////////////////////////////////////////////////////////////////////////
-
 void FTimerNode::ResetAggregatedStats()
 {
 	AggregatedStats = TraceServices::FTimingProfilerAggregatedStats();
diff --git a/Engine/Source/Developer/TraceInsights/Private/Insights/TimingProfiler/ViewModels/TimerNode.h b/Engine/Source/Developer/TraceInsights/Private/Insights/TimingProfiler/ViewModels/TimerNode.h
--- a/Engine/Source/Developer/TraceInsights/Private/Insights/TimingProfiler/ViewModels/TimerNode.h
+++ b/Engine/Source/Developer/TraceInsights/Private/Insights/TimingProfiler/ViewModels/TimerNode.h
@@ -69,10 +69,15 @@ public:
 	FTimerNode(uint32 InTimerId, const TCHAR* InName, ETimerNodeType InType, bool bInIsGroup);
 
 	/** Initialization constructor for the group node. */
-	explicit FTimerNode(const FName InGroupName);
+	explicit FTimerNode(const FText& InDisplayName);
 
 	virtual ~FTimerNode();
 
+	virtual const FText GetDisplayName() const override
+	{
+		return DisplayName;
+	}
+
 	/**
 	 * @return the timer id as provided by analyzer. It can be used as an index.
 	 */
@@ -91,12 +96,20 @@ public:
 	/**
 	 * @return color of the node. Used when showing a graph series for a stats counter.
 	 */
-	FLinearColor GetColor() const
+	virtual FLinearColor GetColor() const override
 	{
 		return Color;
 	}
 
-	void SetColor(const FLinearColor& InColor);
+	void SetColor(const FLinearColor& InColor)
+	{
+		Color = InColor;
+	}
+
+	virtual uint32 GetDefaultSortOrder() const override
+	{
+		return TimerId;
+	}
 
 	bool IsAddedToGraph() const
 	{
@@ -123,6 +136,9 @@ public:
 	bool GetSourceFileAndLine(FString& OutFile, uint32& OutLine) const;
 
 private:
+	/** The display name of the node. */
+	FText DisplayName;
+
 	/** The timer id provided by analyzer. It can also be used as an index. */
 	const uint32 TimerId;
 
diff --git a/Engine/Source/Developer/TraceInsights/Private/Insights/TimingProfiler/Widgets/STimerTreeView.cpp b/Engine/Source/Developer/TraceInsights/Private/Insights/TimingProfiler/Widgets/STimerTreeView.cpp
--- a/Engine/Source/Developer/TraceInsights/Private/Insights/TimingProfiler/Widgets/STimerTreeView.cpp
+++ b/Engine/Source/Developer/TraceInsights/Private/Insights/TimingProfiler/Widgets/STimerTreeView.cpp
@@ -293,7 +293,7 @@ TSharedPtr<SWidget> STimerTreeView::TreeView_GetMenuContent()
 	else if (NumSelectedNodes == 1)
 	{
 		check(SingleSelectedNode.IsValid());
-		FString ItemName = SingleSelectedNode->GetName().ToString();
+		FString ItemName = SingleSelectedNode->GetDisplayName().ToString();
 		const int32 MaxStringLen = 64;
 		if (ItemName.Len() > MaxStringLen)
 		{
@@ -1712,7 +1712,7 @@ void STimerTreeView::ToggleGraphInstanceSeries(TSharedRef<FTimingGraphTrack> Gra
 	{
 		GraphTrack->Show();
 		Series = GraphTrack->AddTimerSeries(TimerId, NodePtr->GetColor());
-		Series->SetName(FText::FromName(NodePtr->GetName()));
+		Series->SetName(NodePtr->GetDisplayName());
 		GraphTrack->SetDirtyFlag();
 	}
 
@@ -1768,7 +1768,7 @@ void STimerTreeView::ToggleGraphFrameStatsSeries(TSharedRef<FTimingGraphTrack> G
 		GraphTrack->Show();
 		Series = GraphTrack->AddFrameStatsTimerSeries(TimerId, FrameType, NodePtr->GetColor());
 		FText SeriesName = FText::Format(LOCTEXT("FrameStatsTimerSeriesName_Fmt", "{0} ({1})"),
-			FText::FromName(NodePtr->GetName()),
+			NodePtr->GetDisplayName(),
 			FFrameTrackDrawHelper::FrameTypeToText(FrameType));
 		Series->SetName(SeriesName);
 		GraphTrack->SetDirtyFlag();
diff --git a/Engine/Source/Developer/TraceInsights/Private/Insights/TimingProfiler/Widgets/STimersView.cpp b/Engine/Source/Developer/TraceInsights/Private/Insights/TimingProfiler/Widgets/STimersView.cpp
--- a/Engine/Source/Developer/TraceInsights/Private/Insights/TimingProfiler/Widgets/STimersView.cpp
+++ b/Engine/Source/Developer/TraceInsights/Private/Insights/TimingProfiler/Widgets/STimersView.cpp
@@ -631,7 +631,7 @@ TSharedPtr<SWidget> STimersView::TreeView_GetMenuContent()
 	else if (NumSelectedNodes == 1)
 	{
 		check(SingleSelectedNode.IsValid());
-		FString ItemName = SingleSelectedNode->GetName().ToString();
+		FString ItemName = SingleSelectedNode->GetDisplayName().ToString();
 		const int32 MaxStringLen = 64;
 		if (ItemName.Len() > MaxStringLen)
 		{
@@ -1730,7 +1730,7 @@ void STimersView::ApplyFiltering()
 
 void STimersView::HandleItemToStringArray(const FTimerNodePtr& FTimerNodePtr, TArray<FString>& OutSearchStrings) const
 {
-	OutSearchStrings.Add(FTimerNodePtr->GetName().GetPlainNameString());
+	OutSearchStrings.Add(FTimerNodePtr->GetDisplayName().ToString());
 }
 
 ////////////////////////////////////////////////////////////////////////////////////////////////////
@@ -2024,8 +2024,7 @@ void STimersView::CreateGroups()
 	{
 		GroupNodes.Reset();
 
-		const FName GroupName(TEXT("All"));
-		FTimerNodePtr GroupPtr = MakeShared<FTimerNode>(GroupName);
+		FTimerNodePtr GroupPtr = MakeShared<FTimerNode>(LOCTEXT("GroupName_All", "All"));
 		GroupNodes.Add(GroupPtr);
 
 		for (const FTimerNodePtr& NodePtr : TimerNodes)
@@ -2044,7 +2043,8 @@ void STimersView::CreateGroups()
 			FTimerNodePtr GroupPtr = GroupNodeSet.FindRef(GroupName);
 			if (!GroupPtr)
 			{
-				GroupPtr = GroupNodeSet.Add(GroupName, MakeShared<FTimerNode>(GroupName));
+				const FText DisplayName = FText::FromName(GroupName);
+				GroupPtr = GroupNodeSet.Add(GroupName, MakeShared<FTimerNode>(DisplayName));
 			}
 			GroupPtr->AddChildAndSetParent(NodePtr);
 			TreeView->SetItemExpansion(GroupPtr, true);
@@ -2062,7 +2062,7 @@ void STimersView::CreateGroups()
 			FTimerNodePtr GroupPtr = GroupNodeSet.FindRef(NodeType);
 			if (!GroupPtr)
 			{
-				const FName GroupName = *TimerNodeTypeHelper::ToText(NodeType).ToString();
+				const FText GroupName = TimerNodeTypeHelper::ToText(NodeType);
 				GroupPtr = GroupNodeSet.Add(NodeType, MakeShared<FTimerNode>(GroupName));
 			}
 			GroupPtr->AddChildAndSetParent(NodePtr);
@@ -2077,12 +2077,12 @@ void STimersView::CreateGroups()
 		TMap<TCHAR, FTimerNodePtr> GroupNodeSet;
 		for (const FTimerNodePtr& NodePtr : TimerNodes)
 		{
-			FString FirstLetterStr(NodePtr->GetName().GetPlainNameString().Left(1).ToUpper());
-			const TCHAR FirstLetter = FirstLetterStr[0];
+			FString FirstLetterStr(NodePtr->GetDisplayName().ToString().Left(1).ToUpper());
+			const TCHAR FirstLetter = FirstLetterStr.Len() > 0 ? FirstLetterStr[0] : TEXT('\0');
 			FTimerNodePtr GroupPtr = GroupNodeSet.FindRef(FirstLetter);
 			if (!GroupPtr)
 			{
-				const FName GroupName(FirstLetterStr);
+				const FText GroupName = FText::FromString(FirstLetterStr);
 				GroupPtr = GroupNodeSet.Add(FirstLetter, MakeShared<FTimerNode>(GroupName));
 			}
 			GroupPtr->AddChildAndSetParent(NodePtr);
@@ -2115,10 +2115,12 @@ void STimersView::CreateGroups()
 			FTimerNodePtr GroupPtr = GroupNodeSet.FindRef(Order);
 			if (!GroupPtr)
 			{
-				const FName GroupName =
-				    (Order == 0) ?          FName(TEXT("Count == 0")) :
-				    (Order < MaxOrder) ?    FName(FString::Printf(TEXT("Count: [%s .. %s)"), Orders[Order - 1], Orders[Order])) :
-				                            FName(FString::Printf(TEXT("Count >= %s"), Orders[MaxOrder - 1]));
+				const FText GroupName =
+					(Order == 0) ?
+						LOCTEXT("GroupName_ZeroCount", "Count == 0") :
+					(Order < MaxOrder) ?
+						FText::Format(LOCTEXT("GroupName_CountFmt1", "Count: [{0} .. {1})"), FText::FromString(Orders[Order - 1]), FText::FromString(Orders[Order])) :
+						FText::Format(LOCTEXT("GroupName_CountFmt2", "Count >= {0}"), FText::FromString(Orders[MaxOrder - 1]));
 				GroupPtr = GroupNodeSet.Add(Order, MakeShared<FTimerNode>(GroupName));
 			}
 			GroupPtr->AddChildAndSetParent(NodePtr);
@@ -2126,16 +2128,6 @@ void STimersView::CreateGroups()
 		GroupNodeSet.KeySort([](const uint32& A, const uint32& B) { return A > B; }); // sort groups by order
 		GroupNodeSet.GenerateValueArray(GroupNodes);
 	}
-	// Creates one group for each logarithmic range ie. 0.001 - 0.01, 0.01 - 0.1, 0.1 - 1.0, 1.0 - 10.0, etc.
-	else if (GroupingMode == ETimerGroupingMode::ByTotalInclusiveTime)
-	{
-		unimplemented();
-	}
-	// Creates one group for each logarithmic range ie. 0.001 - 0.01, 0.01 - 0.1, 0.1 - 1.0, 1.0 - 10.0, etc.
-	else if (GroupingMode == ETimerGroupingMode::ByTotalExclusiveTime)
-	{
-		unimplemented();
-	}
 	else
 	{
 		unimplemented();
@@ -2871,7 +2863,6 @@ void STimersView::RebuildTree(bool bResync)
 					}
 
 					TimerNodePtr = MakeShared<FTimerNode>(Timer.Id, *StringBuilder, Type, false);
-					TimerNodePtr->SetDefaultSortOrder(TimerIndex + 1);
 
 					const bool bIsTimerAddedToGraphs =
 						(FrameTrack.IsValid() && FrameTrack->HasAnySeriesForTimer(Timer.Id)) ||
@@ -3106,7 +3097,7 @@ void STimersView::ToggleGraphInstanceSeries(TSharedRef<FTimingGraphTrack> GraphT
 	{
 		GraphTrack->Show();
 		Series = GraphTrack->AddTimerSeries(TimerId, NodePtr->GetColor());
-		Series->SetName(FText::FromName(NodePtr->GetName()));
+		Series->SetName(NodePtr->GetDisplayName());
 		GraphTrack->SetDirtyFlag();
 	}
 
@@ -3162,7 +3153,7 @@ void STimersView::ToggleGraphFrameStatsSeries(TSharedRef<FTimingGraphTrack> Grap
 		GraphTrack->Show();
 		Series = GraphTrack->AddFrameStatsTimerSeries(TimerId, FrameType, NodePtr->GetColor());
 		FText SeriesName = FText::Format(LOCTEXT("FrameStatsTimerSeriesName_Fmt", "{0} ({1})"),
-			FText::FromName(NodePtr->GetName()),
+			NodePtr->GetDisplayName(),
 			FFrameTrackDrawHelper::FrameTypeToText(FrameType));
 		Series->SetName(SeriesName);
 		GraphTrack->SetDirtyFlag();
diff --git a/Engine/Source/Developer/TraceInsights/Private/Insights/TimingProfiler/Widgets/STimersViewTooltip.cpp b/Engine/Source/Developer/TraceInsights/Private/Insights/TimingProfiler/Widgets/STimersViewTooltip.cpp
--- a/Engine/Source/Developer/TraceInsights/Private/Insights/TimingProfiler/Widgets/STimersViewTooltip.cpp
+++ b/Engine/Source/Developer/TraceInsights/Private/Insights/TimingProfiler/Widgets/STimersViewTooltip.cpp
@@ -260,16 +260,16 @@ TSharedPtr<SToolTip> STimersViewTooltip::GetRowTooltip(const TSharedPtr<FTimerNo
 						SNew(STextBlock)
 						.WrapTextAt(1024.0f)
 						.WrappingPolicy(ETextWrappingPolicy::AllowPerCharacterWrapping)
-						.Text(FText::FromName(TimerNodePtr->GetName()))
+						.Text(TimerNodePtr->GetDisplayName())
 						.TextStyle(FInsightsStyle::Get(), TEXT("TreeTable.Tooltip"))
 					]
 
-					// Timer Type: [Type]
+					// Item Type: [Type]
 					+ SGridPanel::Slot(0, 2)
 					.Padding(2.0f)
 					[
 						SNew(STextBlock)
-						.Text(LOCTEXT("TT_Type", "Type:"))
+						.Text(LOCTEXT("TT_Type", "Item Type:"))
 						.TextStyle(FInsightsStyle::Get(), TEXT("TreeTable.TooltipBold"))
 					]
 					+ SGridPanel::Slot(1, 2)
