[Insights] Fixed the deprecated FName usage in FStatsNode and in SStatsView. Counters will not create FNames anymore.

#Synced-CL 51334642

diff --git a/Engine/Source/Developer/TraceInsights/Private/Insights/TimingProfiler/ViewModels/StatsGroupingAndSorting.cpp b/Engine/Source/Developer/TraceInsights/Private/Insights/TimingProfiler/ViewModels/StatsGroupingAndSorting.cpp
--- a/Engine/Source/Developer/TraceInsights/Private/Insights/TimingProfiler/ViewModels/StatsGroupingAndSorting.cpp
+++ b/Engine/Source/Developer/TraceInsights/Private/Insights/TimingProfiler/ViewModels/StatsGroupingAndSorting.cpp
@@ -12,7 +12,7 @@
 #define INSIGHTS_ENSURE(...)
 
 // Sort by name (ascending).
-#define INSIGHTS_DEFAULT_SORTING_NODES(A, B) return A->GetName().LexicalLess(B->GetName());
+#define INSIGHTS_DEFAULT_SORTING_NODES(A, B) return (A->GetDisplayName().CompareTo(B->GetDisplayName()) < 0);
 //#define INSIGHTS_DEFAULT_SORTING_NODES(A, B) return A->GetDefaultSortOrder() < B->GetDefaultSortOrder();
 
 #define LOCTEXT_NAMESPACE "UE::Insights::TimingProfiler::StatsNode"
diff --git a/Engine/Source/Developer/TraceInsights/Private/Insights/TimingProfiler/ViewModels/StatsNode.cpp b/Engine/Source/Developer/TraceInsights/Private/Insights/TimingProfiler/ViewModels/StatsNode.cpp
--- a/Engine/Source/Developer/TraceInsights/Private/Insights/TimingProfiler/ViewModels/StatsNode.cpp
+++ b/Engine/Source/Developer/TraceInsights/Private/Insights/TimingProfiler/ViewModels/StatsNode.cpp
@@ -18,9 +18,42 @@ INSIGHTS_IMPLEMENT_RTTI(FStatsNode)
 
 ////////////////////////////////////////////////////////////////////////////////////////////////////
 
-void FStatsNode::SetColor(const FLinearColor& InColor)
+FStatsNode::FStatsNode(uint32 InCounterId, const FText& InDisplayName, const FName InMetaGroupName, EStatsNodeType InType, EStatsNodeDataType InDataType)
+	: FBaseTreeNode(InType == EStatsNodeType::Group)
+	, DisplayName(InDisplayName)
+	, CounterId(InCounterId)
+	, MetaGroupName(InMetaGroupName)
+	, Type(InType)
+	, DataType(InDataType)
+	, bIsAddedToGraph(false)
+{
+	const uint32 HashColor = GetCounterId() * 0x2c2c57ed;
+	Color.R = ((HashColor >> 16) & 0xFF) / 255.0f;
+	Color.G = ((HashColor >>  8) & 0xFF) / 255.0f;
+	Color.B = ((HashColor      ) & 0xFF) / 255.0f;
+	Color.A = 1.0f;
+
+	ResetAggregatedStats();
+}
+
+////////////////////////////////////////////////////////////////////////////////////////////////////
+
+FStatsNode::FStatsNode(const FText& InDisplayName)
+	: FBaseTreeNode(true)
+	, DisplayName(InDisplayName)
+	, CounterId(InvalidCounterId)
+	, Type(EStatsNodeType::Group)
+	, DataType(EStatsNodeDataType::InvalidOrMax)
+	, Color(0.0f, 0.0f, 0.0f, 1.0f)
+	, bIsAddedToGraph(false)
+{
+	ResetAggregatedStats();
+}
+
+////////////////////////////////////////////////////////////////////////////////////////////////////
+
+FStatsNode::~FStatsNode()
 {
-	Color = InColor;
 }
 
 ////////////////////////////////////////////////////////////////////////////////////////////////////
diff --git a/Engine/Source/Developer/TraceInsights/Private/Insights/TimingProfiler/ViewModels/StatsNode.h b/Engine/Source/Developer/TraceInsights/Private/Insights/TimingProfiler/ViewModels/StatsNode.h
--- a/Engine/Source/Developer/TraceInsights/Private/Insights/TimingProfiler/ViewModels/StatsNode.h
+++ b/Engine/Source/Developer/TraceInsights/Private/Insights/TimingProfiler/ViewModels/StatsNode.h
@@ -141,39 +141,25 @@ public:
 
 public:
 	/** Initialization constructor for the stats node. */
-	FStatsNode(uint32 InCounterId, const FName InName, const FName InMetaGroupName, EStatsNodeType InType, EStatsNodeDataType InDataType)
-		: FBaseTreeNode(InName, InType == EStatsNodeType::Group)
-		, CounterId(InCounterId)
-		, MetaGroupName(InMetaGroupName)
-		, Type(InType)
-		, DataType(InDataType)
-		, bIsAddedToGraph(false)
-	{
-		const uint32 HashColor = GetCounterId() * 0x2c2c57ed;
-		Color.R = ((HashColor >> 16) & 0xFF) / 255.0f;
-		Color.G = ((HashColor >> 8) & 0xFF) / 255.0f;
-		Color.B = ((HashColor) & 0xFF) / 255.0f;
-		Color.A = 1.0;
-
-		ResetAggregatedStats();
-	}
+	FStatsNode(uint32 InCounterId, const FText& InDisplayName, const FName InMetaGroupName, EStatsNodeType InType, EStatsNodeDataType InDataType);
 
 	/** Initialization constructor for the group node. */
-	explicit FStatsNode(const FName InGroupName)
-		: FBaseTreeNode(InGroupName, true)
-		, CounterId(InvalidCounterId)
-		, Type(EStatsNodeType::Group)
-		, DataType(EStatsNodeDataType::InvalidOrMax)
-		, Color(0.0, 0.0, 0.0, 1.0)
-		, bIsAddedToGraph(false)
+	explicit FStatsNode(const FText& InDisplayName);
+
+	virtual ~FStatsNode();
+
+	virtual const FText GetDisplayName() const override
 	{
-		ResetAggregatedStats();
+		return DisplayName;
 	}
 
 	/**
 	 * @return the counter id as provided by analyzer.
 	 */
-	uint32 GetCounterId() const { return CounterId; }
+	uint32 GetCounterId() const
+	{
+		return CounterId;
+	}
 
 	/**
 	 * @return the name of the meta group that this stats node belongs to, taken from the metadata.
@@ -209,12 +195,20 @@ public:
 	/**
 	 * @return color of the node. Used when showing a graph series for a stats counter.
 	 */
-	FLinearColor GetColor() const
+	virtual FLinearColor GetColor() const override
 	{
 		return Color;
 	}
 
-	void SetColor(const FLinearColor& InColor);
+	void SetColor(const FLinearColor& InColor)
+	{
+		Color = InColor;
+	}
+
+	virtual uint32 GetDefaultSortOrder() const override
+	{
+		return CounterId;
+	}
 
 	bool IsAddedToGraph() const
 	{
@@ -252,6 +246,9 @@ public:
 	const FText GetTextForAggregatedStatsDiff(bool bForTooltip = false) const;
 
 private:
+	/** The display name of the node. */
+	FText DisplayName;
+
 	/** The counter id provided by the analyzer. */
 	uint32 CounterId;
 
diff --git a/Engine/Source/Developer/TraceInsights/Private/Insights/TimingProfiler/Widgets/SStatsView.cpp b/Engine/Source/Developer/TraceInsights/Private/Insights/TimingProfiler/Widgets/SStatsView.cpp
--- a/Engine/Source/Developer/TraceInsights/Private/Insights/TimingProfiler/Widgets/SStatsView.cpp
+++ b/Engine/Source/Developer/TraceInsights/Private/Insights/TimingProfiler/Widgets/SStatsView.cpp
@@ -434,7 +434,7 @@ TSharedPtr<SWidget> SStatsView::TreeView_GetMenuContent()
 	else if (NumSelectedNodes == 1)
 	{
 		check(SingleSelectedNode.IsValid());
-		FString ItemName = SingleSelectedNode->GetName().ToString();
+		FString ItemName = SingleSelectedNode->GetDisplayName().ToString();
 		const int32 MaxStringLen = 64;
 		if (ItemName.Len() > MaxStringLen)
 		{
@@ -1079,7 +1079,7 @@ void SStatsView::ApplyFiltering()
 
 void SStatsView::HandleItemToStringArray(const FStatsNodePtr& FStatsNodePtr, TArray<FString>& OutSearchStrings) const
 {
-	OutSearchStrings.Add(FStatsNodePtr->GetName().GetPlainNameString());
+	OutSearchStrings.Add(FStatsNodePtr->GetDisplayName().ToString());
 }
 
 ////////////////////////////////////////////////////////////////////////////////////////////////////
@@ -1358,8 +1358,7 @@ void SStatsView::CreateGroups()
 	{
 		GroupNodes.Reset();
 
-		const FName GroupName(TEXT("All"));
-		FStatsNodePtr GroupPtr = MakeShared<FStatsNode>(GroupName);
+		FStatsNodePtr GroupPtr = MakeShared<FStatsNode>(LOCTEXT("GroupName_All", "All"));
 		GroupNodes.Add(GroupPtr);
 
 		for (const FStatsNodePtr& NodePtr : StatsNodes)
@@ -1378,7 +1377,8 @@ void SStatsView::CreateGroups()
 			FStatsNodePtr GroupPtr = GroupNodeSet.FindRef(GroupName);
 			if (!GroupPtr)
 			{
-				GroupPtr = GroupNodeSet.Add(GroupName, MakeShared<FStatsNode>(GroupName));
+				const FText DisplayName = FText::FromName(GroupName);
+				GroupPtr = GroupNodeSet.Add(GroupName, MakeShared<FStatsNode>(DisplayName));
 			}
 			GroupPtr->AddChildAndSetParent(NodePtr);
 			TreeView->SetItemExpansion(GroupPtr, true);
@@ -1396,7 +1396,7 @@ void SStatsView::CreateGroups()
 			FStatsNodePtr GroupPtr = GroupNodeSet.FindRef(NodeType);
 			if (!GroupPtr)
 			{
-				const FName GroupName = *StatsNodeTypeHelper::ToText(NodeType).ToString();
+				const FText GroupName = StatsNodeTypeHelper::ToText(NodeType);
 				GroupPtr = GroupNodeSet.Add(NodeType, MakeShared<FStatsNode>(GroupName));
 			}
 			GroupPtr->AddChildAndSetParent(NodePtr);
@@ -1415,7 +1415,7 @@ void SStatsView::CreateGroups()
 			FStatsNodePtr GroupPtr = GroupNodeSet.FindRef(DataType);
 			if (!GroupPtr)
 			{
-				const FName GroupName = *StatsNodeDataTypeHelper::ToText(DataType).ToString();
+				const FText GroupName = StatsNodeDataTypeHelper::ToText(DataType);
 				GroupPtr = GroupNodeSet.Add(DataType, MakeShared<FStatsNode>(GroupName));
 			}
 			GroupPtr->AddChildAndSetParent(NodePtr);
@@ -1430,12 +1430,12 @@ void SStatsView::CreateGroups()
 		TMap<TCHAR, FStatsNodePtr> GroupNodeSet;
 		for (const FStatsNodePtr& NodePtr : StatsNodes)
 		{
-			FString FirstLetterStr(NodePtr->GetName().GetPlainNameString().Left(1).ToUpper());
-			const TCHAR FirstLetter = FirstLetterStr[0];
+			FString FirstLetterStr(NodePtr->GetDisplayName().ToString().Left(1).ToUpper());
+			const TCHAR FirstLetter = FirstLetterStr.Len() > 0 ? FirstLetterStr[0] : TEXT('\0');
 			FStatsNodePtr GroupPtr = GroupNodeSet.FindRef(FirstLetter);
 			if (!GroupPtr)
 			{
-				const FName GroupName(FirstLetterStr);
+				const FText GroupName = FText::FromString(FirstLetterStr);
 				GroupPtr = GroupNodeSet.Add(FirstLetter, MakeShared<FStatsNode>(GroupName));
 			}
 			GroupPtr->AddChildAndSetParent(NodePtr);
@@ -1468,10 +1468,12 @@ void SStatsView::CreateGroups()
 			FStatsNodePtr GroupPtr = GroupNodeSet.FindRef(Order);
 			if (!GroupPtr)
 			{
-				const FName GroupName =
-				    (Order == 0) ?          FName(TEXT("Count == 0")) :
-				    (Order < MaxOrder) ?    FName(FString::Printf(TEXT("Count: [%s .. %s)"), Orders[Order - 1], Orders[Order])) :
-				                            FName(FString::Printf(TEXT("Count >= %s"), Orders[MaxOrder - 1]));
+				const FText GroupName =
+					(Order == 0) ?
+						LOCTEXT("GroupName_ZeroCount", "Count == 0") :
+					(Order < MaxOrder) ?
+						FText::Format(LOCTEXT("GroupName_CountFmt1", "Count: [{0} .. {1})"), FText::FromString(Orders[Order - 1]), FText::FromString(Orders[Order])) :
+						FText::Format(LOCTEXT("GroupName_CountFmt2", "Count >= {0}"), FText::FromString(Orders[MaxOrder - 1]));
 				GroupPtr = GroupNodeSet.Add(Order, MakeShared<FStatsNode>(GroupName));
 			}
 			GroupPtr->AddChildAndSetParent(NodePtr);
@@ -2072,12 +2074,12 @@ void SStatsView::RebuildTree(bool bResync)
 					FStatsNodePtr NodePtr = StatsNodesIdMap.FindRef(CounterId);
 					if (!NodePtr)
 					{
-						FName Name(Counter.GetName());
+						++CounterIndex;
+						const FText DisplayName = FText::FromName(Counter.GetName());
 						const FName Group = Counter.GetGroup() ? Counter.GetGroup() : DeduceGroup(Counter);
 						const EStatsNodeType Type = EStatsNodeType::Counter;
 						const EStatsNodeDataType DataType = Counter.IsFloatingPoint() ? EStatsNodeDataType::Double : EStatsNodeDataType::Int64;
-						NodePtr = MakeShared<FStatsNode>(CounterId, Name, Group, Type, DataType);
-						NodePtr->SetDefaultSortOrder(++CounterIndex);
+						NodePtr = MakeShared<FStatsNode>(CounterId, DisplayName, Group, Type, DataType);
 						UpdateNode(NodePtr);
 						StatsNodes.Add(NodePtr);
 						StatsNodesIdMap.Add(CounterId, NodePtr);
@@ -2454,7 +2456,7 @@ void SStatsView::ContextMenu_ExportValues_Execute() const
 	}
 
 	const uint32 CounterId = SelectedNode->GetCounterId();
-	FString CounterName = SelectedNode->GetName().ToString();
+	FString CounterName = SelectedNode->GetDisplayName().ToString();
 
 	const FString DialogTitle = LOCTEXT("ExportValues_Title", "Export Counter Values").ToString();
 
@@ -2523,7 +2525,7 @@ void SStatsView::ContextMenu_ExportOps_Execute() const
 	}
 
 	const uint32 CounterId = SelectedNode->GetCounterId();
-	FString CounterName = SelectedNode->GetName().ToString();
+	FString CounterName = SelectedNode->GetDisplayName().ToString();
 
 	const FString DialogTitle = LOCTEXT("ExportOps_Title", "Export Counter Ops").ToString();
 
diff --git a/Engine/Source/Developer/TraceInsights/Private/Insights/TimingProfiler/Widgets/SStatsViewTooltip.cpp b/Engine/Source/Developer/TraceInsights/Private/Insights/TimingProfiler/Widgets/SStatsViewTooltip.cpp
--- a/Engine/Source/Developer/TraceInsights/Private/Insights/TimingProfiler/Widgets/SStatsViewTooltip.cpp
+++ b/Engine/Source/Developer/TraceInsights/Private/Insights/TimingProfiler/Widgets/SStatsViewTooltip.cpp
@@ -164,16 +164,16 @@ TSharedPtr<SToolTip> SStatsViewTooltip::GetRowTooltip(const TSharedPtr<FStatsNod
 						SNew(STextBlock)
 						.WrapTextAt(512.0f)
 						.WrappingPolicy(ETextWrappingPolicy::AllowPerCharacterWrapping)
-						.Text(FText::FromName(StatsNodePtr->GetName()))
+						.Text(StatsNodePtr->GetDisplayName())
 						.TextStyle(FInsightsStyle::Get(), TEXT("TreeTable.Tooltip"))
 					]
 
-					// Counter Type: [Type]
+					// Item Type: [Type]
 					+ SGridPanel::Slot(0, 3)
 					.Padding(2.0f)
 					[
 						SNew(STextBlock)
-						.Text(LOCTEXT("TT_Type", "Node Type:"))
+						.Text(LOCTEXT("TT_Type", "Item Type:"))
 						.TextStyle(FInsightsStyle::Get(), TEXT("TreeTable.TooltipBold"))
 					]
 					+ SGridPanel::Slot(1, 3)
