Adding 'UnrealInsights' feature test category to ini file and common test name string variables

#rnx
#tests utilized new string variable in a UnitTest and ran locally

#Synced-CL 51376167

diff --git a/Engine/Source/Developer/TraceInsights/Private/Insights/Tests/UnitTests.cpp b/Engine/Source/Developer/TraceInsights/Private/Insights/Tests/UnitTests.cpp
--- a/Engine/Source/Developer/TraceInsights/Private/Insights/Tests/UnitTests.cpp
+++ b/Engine/Source/Developer/TraceInsights/Private/Insights/Tests/UnitTests.cpp
@@ -2,8 +2,11 @@
 
 #include "CoreTypes.h"
 #include "Misc/AutomationTest.h"
+#include "Insights/Tests/InsightsTestUtils.h"
 
-IMPLEMENT_SIMPLE_AUTOMATION_TEST(FTraceInsightsUnitTest, "System.Insights.UnitTest", EAutomationTestFlags_ApplicationContextMask | EAutomationTestFlags::SmokeFilter)
+using namespace UE::Insights::Automation::TestNames;
+
+IMPLEMENT_SIMPLE_AUTOMATION_TEST(FTraceInsightsUnitTest, UnitTests.GetData(), EAutomationTestFlags_ApplicationContextMask | EAutomationTestFlags::SmokeFilter)
 
 bool FTraceInsightsUnitTest::RunTest(const FString& Parameters)
 {
diff --git a/Engine/Source/Developer/TraceInsights/Public/Insights/Tests/InsightsTestUtils.h b/Engine/Source/Developer/TraceInsights/Public/Insights/Tests/InsightsTestUtils.h
--- a/Engine/Source/Developer/TraceInsights/Public/Insights/Tests/InsightsTestUtils.h
+++ b/Engine/Source/Developer/TraceInsights/Public/Insights/Tests/InsightsTestUtils.h
@@ -10,6 +10,16 @@
 
 #define UE_API TRACEINSIGHTS_API
 
+namespace UE::Insights::Automation::TestNames
+{
+	const FStringView EditorTests = TEXTVIEW("UnrealInsights.Editor.");
+	const FStringView ClientTests = TEXTVIEW("UnrealInsights.Client.");
+	const FStringView ViewerStaticTraceTests = TEXTVIEW("UnrealInsights.Trace.Analysis.");
+	const FStringView ViewerLiveTraceTests = TEXTVIEW("UnrealInsights.Analysis.");
+	const FStringView HubTests = TEXTVIEW("UnrealInsights.Hub.");
+	const FStringView UnitTests = TEXTVIEW("UnrealInsights.UnitTest.");
+}
+
 class FAutomationTestBase;
 
 class FInsightsTestUtils
