[Insights] Added ObjectProfiler and the "Objects" table tree view -- visualization for a UObject snapshot.

#Synced-CL 51457572

diff --git a/Engine/Source/Developer/TraceInsights/Private/Insights/InsightsStyle.cpp b/Engine/Source/Developer/TraceInsights/Private/Insights/InsightsStyle.cpp
--- a/Engine/Source/Developer/TraceInsights/Private/Insights/InsightsStyle.cpp
+++ b/Engine/Source/Developer/TraceInsights/Private/Insights/InsightsStyle.cpp
@@ -91,10 +91,13 @@ void FInsightsStyle::FStyle::SyncParentStyles()
 
 ////////////////////////////////////////////////////////////////////////////////////////////////////
 
-#define EDITOR_IMAGE_BRUSH(RelativePath, ...) IMAGE_BRUSH("../../../Editor/Slate/" RelativePath, __VA_ARGS__)
-#define EDITOR_IMAGE_BRUSH_SVG(RelativePath, ...) IMAGE_BRUSH_SVG("../../../Editor/Slate/" RelativePath, __VA_ARGS__)
-#define EDITOR_BOX_BRUSH(RelativePath, ...) BOX_BRUSH("../../../Editor/Slate/" RelativePath, __VA_ARGS__)
-#define EDITOR_BORDER_BRUSH(RelativePath, ...) BORDER_BRUSH("../../../Editor/Slate/" RelativePath, __VA_ARGS__)
+#define EDITOR_IMAGE_BRUSH(RelativePath, ...)      IMAGE_BRUSH      ("../../../Editor/Slate/" RelativePath, __VA_ARGS__)
+#define EDITOR_IMAGE_BRUSH_SVG(RelativePath, ...)  IMAGE_BRUSH_SVG  ("../../../Editor/Slate/" RelativePath, __VA_ARGS__)
+#define EDITOR_BOX_BRUSH(RelativePath, ...)        BOX_BRUSH        ("../../../Editor/Slate/" RelativePath, __VA_ARGS__)
+#define EDITOR_BOX_BRUSH_SVG(RelativePath, ...)    BOX_BRUSH_SVG    ("../../../Editor/Slate/" RelativePath, __VA_ARGS__)
+#define EDITOR_BORDER_BRUSH(RelativePath, ...)     BORDER_BRUSH     ("../../../Editor/Slate/" RelativePath, __VA_ARGS__)
+#define EDITOR_BORDER_BRUSH_SVG(RelativePath, ...) BORDER_BRUSH_SVG ("../../../Editor/Slate/" RelativePath, __VA_ARGS__)
+
 #define TODO_IMAGE_BRUSH(...) EDITOR_IMAGE_BRUSH_SVG("Starship/Common/StaticMesh", __VA_ARGS__)
 
 void FInsightsStyle::FStyle::Initialize()
@@ -153,7 +156,9 @@ void FInsightsStyle::FStyle::Initialize()
 
 	Set("Icons.TasksView", new IMAGE_BRUSH_SVG("Tasks", Icon16x16));
 	Set("Icons.PackagesView", new EDITOR_IMAGE_BRUSH_SVG("Starship/Common/Spreadsheet", Icon16x16));
-	
+
+	Set("Icons.ObjectsView", new EDITOR_IMAGE_BRUSH_SVG("Starship/Common/Spreadsheet", Icon16x16));
+
 	Set("Icons.PanelMenu.Toolbar", new CORE_IMAGE_BRUSH_SVG("Starship/Common/Menu", Icon20x20));
 
 	//////////////////////////////////////////////////
@@ -231,6 +236,8 @@ void FInsightsStyle::FStyle::Initialize()
 	Set("Icons.TagSet.Systems", new IMAGE_BRUSH_SVG("MemTagSet_Systems", Icon16x16));
 	Set("Icons.TagSet.Assets", new IMAGE_BRUSH_SVG("MemTagSet_Assets", Icon16x16));
 	Set("Icons.TagSet.AssetClasses", new IMAGE_BRUSH_SVG("MemTagSet_AssetClasses", Icon16x16));
+	Set("Icons.TagSet.UObjectClasses", new IMAGE_BRUSH_SVG("MemTagSet_AssetClasses", Icon16x16));
+	Set("Icons.TagSet.CodeOrContent", new IMAGE_BRUSH_SVG("MemTagSet_AssetClasses", Icon16x16));
 
 	Set("Icons.BudgetSettings", new IMAGE_BRUSH_SVG("BudgetSettings", Icon16x16));
 	Set("Icons.TimeMarkerSettings", new IMAGE_BRUSH_SVG("TimeMarkerSettings", Icon16x16));
@@ -292,7 +299,13 @@ void FInsightsStyle::FStyle::Initialize()
 	Set("Icons.SystemMemTag.TreeItem", new IMAGE_BRUSH_SVG("MemTag_System_12", Icon12x12));
 	Set("Icons.AssetMemTag.TreeItem", new IMAGE_BRUSH_SVG("MemTag_Asset_12", Icon12x12));
 	Set("Icons.ClassMemTag.TreeItem", new IMAGE_BRUSH_SVG("MemTag_Class_12", Icon12x12));
+
 	Set("Icons.UObject.TreeItem", new IMAGE_BRUSH_SVG("UObject_12", Icon12x12));
+	Set("Icons.UField.TreeItem", new IMAGE_BRUSH_SVG("UField_12", Icon12x12)); // UField -> UObject
+	Set("Icons.UStruct.TreeItem", new IMAGE_BRUSH_SVG("UStruct_12", Icon12x12)); // UStruct -> UField
+	Set("Icons.UClass.TreeItem", new IMAGE_BRUSH_SVG("UClass_12", Icon12x12)); // UClass -> UStruct
+	Set("Icons.UPackage.TreeItem", new IMAGE_BRUSH_SVG("UPackage_12", Icon12x12)); // UPackage -> UObject
+	Set("Icons.AActor.TreeItem", new IMAGE_BRUSH_SVG("AActor_12", Icon12x12)); // AActor -> UObject
 
 	Set("Icons.HasGraph.TreeItem", new IMAGE_BRUSH_SVG("RoundedBullet", Icon12x12));
 
diff --git a/Engine/Source/Developer/TraceInsights/Private/Insights/ObjectProfiler/ObjectProfilerManager.cpp b/Engine/Source/Developer/TraceInsights/Private/Insights/ObjectProfiler/ObjectProfilerManager.cpp
new file mode 100644
--- /dev/null
+++ b/Engine/Source/Developer/TraceInsights/Private/Insights/ObjectProfiler/ObjectProfilerManager.cpp
@@ -0,0 +1,263 @@
+// Copyright Epic Games, Inc. All Rights Reserved.
+
+#include "ObjectProfilerManager.h"
+
+#include "Features/IModularFeatures.h"
+#include "Framework/Docking/TabManager.h"
+#include "Modules/ModuleManager.h"
+#include "Widgets/Docking/SDockTab.h"
+
+// TraceServices
+#include "Common/ProviderLock.h"
+#include "TraceServices/Model/ObjectProvider.h"
+
+// TraceInsights
+#include "Insights/InsightsStyle.h"
+#include "Insights/ObjectProfiler/ViewModels/ObjectTable.h"
+#include "Insights/ObjectProfiler/Widgets/SObjectTableTreeView.h"
+
+#define LOCTEXT_NAMESPACE "UE::Insights::ObjectProfiler"
+
+namespace UE::Insights::ObjectProfiler
+{
+
+DEFINE_LOG_CATEGORY(LogObjectProfiler);
+
+////////////////////////////////////////////////////////////////////////////////////////////////////
+
+const FName FObjectProfilerTabs::ObjectsTableTreeViewTabId(TEXT("ObjectsTableTreeView"));
+
+////////////////////////////////////////////////////////////////////////////////////////////////////
+
+TSharedPtr<FObjectProfilerManager> FObjectProfilerManager::Instance = nullptr;
+
+////////////////////////////////////////////////////////////////////////////////////////////////////
+
+TSharedPtr<FObjectProfilerManager> FObjectProfilerManager::Get()
+{
+	return FObjectProfilerManager::Instance;
+}
+
+////////////////////////////////////////////////////////////////////////////////////////////////////
+
+TSharedPtr<FObjectProfilerManager> FObjectProfilerManager::CreateInstance()
+{
+	ensure(!FObjectProfilerManager::Instance.IsValid());
+	if (FObjectProfilerManager::Instance.IsValid())
+	{
+		FObjectProfilerManager::Instance.Reset();
+	}
+
+	FObjectProfilerManager::Instance = MakeShared<FObjectProfilerManager>(FInsightsManager::Get()->GetCommandList());
+
+	return FObjectProfilerManager::Instance;
+}
+
+////////////////////////////////////////////////////////////////////////////////////////////////////
+
+FObjectProfilerManager::FObjectProfilerManager(TSharedRef<FUICommandList> InCommandList)
+	: bIsInitialized(false)
+	, bIsAvailable(false)
+{
+}
+
+////////////////////////////////////////////////////////////////////////////////////////////////////
+
+void FObjectProfilerManager::Initialize(IUnrealInsightsModule& InsightsModule)
+{
+	ensure(!bIsInitialized);
+	if (bIsInitialized)
+	{
+		return;
+	}
+	bIsInitialized = true;
+
+	// Register tick functions.
+	OnTick = FTickerDelegate::CreateSP(this, &FObjectProfilerManager::Tick);
+	OnTickHandle = FTSTicker::GetCoreTicker().AddTicker(OnTick, 0.0f);
+
+	FOnRegisterMajorTabExtensions* TimingProfilerLayoutExtension = InsightsModule.FindMajorTabLayoutExtension(FInsightsManagerTabs::TimingProfilerTabId);
+	if (TimingProfilerLayoutExtension)
+	{
+		TimingProfilerLayoutExtension->AddRaw(this, &FObjectProfilerManager::RegisterTimingProfilerLayoutExtensions);
+	}
+
+	FInsightsManager::Get()->GetSessionChangedEvent().AddSP(this, &FObjectProfilerManager::OnSessionChanged);
+	OnSessionChanged();
+}
+
+////////////////////////////////////////////////////////////////////////////////////////////////////
+
+void FObjectProfilerManager::Shutdown()
+{
+	if (!bIsInitialized)
+	{
+		return;
+	}
+	bIsInitialized = false;
+
+	FInsightsManager::Get()->GetSessionChangedEvent().RemoveAll(this);
+
+	// Unregister tick function.
+	FTSTicker::RemoveTicker(OnTickHandle);
+
+	FObjectProfilerManager::Instance.Reset();
+}
+
+////////////////////////////////////////////////////////////////////////////////////////////////////
+
+FObjectProfilerManager::~FObjectProfilerManager()
+{
+	ensure(!bIsInitialized);
+}
+
+////////////////////////////////////////////////////////////////////////////////////////////////////
+
+void FObjectProfilerManager::RegisterMajorTabs(IUnrealInsightsModule& InsightsModule)
+{
+}
+
+////////////////////////////////////////////////////////////////////////////////////////////////////
+
+void FObjectProfilerManager::UnregisterMajorTabs()
+{
+}
+
+////////////////////////////////////////////////////////////////////////////////////////////////////
+
+bool FObjectProfilerManager::Tick(float DeltaTime)
+{
+	// Check if the session has UObject snapshots, but not too often.
+	if (!bIsAvailable && AvailabilityCheck.Tick())
+	{
+		bool bShouldBeAvailable = false;
+
+		TSharedPtr<const TraceServices::IAnalysisSession> Session = FInsightsManager::Get()->GetSession();
+		if (Session.IsValid())
+		{
+			const TraceServices::IObjectProvider* ObjectProvider = TraceServices::ReadObjectProvider(*Session.Get());
+			if (ObjectProvider)
+			{
+				TraceServices::FProviderReadScopeLock ProviderReadScope(*ObjectProvider);
+
+				TSharedPtr<FTabManager> TabManagerShared = TimingTabManager.Pin();
+				if (ObjectProvider->GetNumSnapshots() > 0 && TabManagerShared.IsValid())
+				{
+					bIsAvailable = true;
+					TabManagerShared->TryInvokeTab(FObjectProfilerTabs::ObjectsTableTreeViewTabId, true);
+				}
+
+				if (Session->IsAnalysisComplete())
+				{
+					// Never check again during this session.
+					AvailabilityCheck.Disable();
+				}
+			}
+		}
+		else
+		{
+			// Do not check again until the next session changed event (see OnSessionChanged).
+			AvailabilityCheck.Disable();
+		}
+	}
+
+	return true;
+}
+
+////////////////////////////////////////////////////////////////////////////////////////////////////
+
+void FObjectProfilerManager::OnSessionChanged()
+{
+	bIsAvailable = false;
+	if (FInsightsManager::Get()->GetSession().IsValid())
+	{
+		AvailabilityCheck.Enable(0.5);
+	}
+	else
+	{
+		AvailabilityCheck.Disable();
+	}
+}
+
+////////////////////////////////////////////////////////////////////////////////////////////////////
+
+void FObjectProfilerManager::RegisterTimingProfilerLayoutExtensions(FInsightsMajorTabExtender& InOutExtender)
+{
+	TimingTabManager = InOutExtender.GetTabManager();
+
+	FMinorTabConfig& MinorTabConfig = InOutExtender.AddMinorTabConfig();
+	MinorTabConfig.TabId = FObjectProfilerTabs::ObjectsTableTreeViewTabId;
+	MinorTabConfig.TabLabel = LOCTEXT("ObjectsTableTreeViewTabTitle", "Objects");
+	MinorTabConfig.TabTooltip = LOCTEXT("ObjectsTableTreeViewTabTitleTooltip", "Opens the Objects Table Tree View tab.");
+	MinorTabConfig.TabIcon = FSlateIcon(FInsightsStyle::GetStyleSetName(), "Icons.ObjectsView");
+	MinorTabConfig.OnSpawnTab = FOnSpawnTab::CreateRaw(this, &FObjectProfilerManager::SpawnTab_ObjectsTableTreeView);
+	MinorTabConfig.CanSpawnTab = FCanSpawnTab::CreateRaw(this, &FObjectProfilerManager::CanSpawnTab_ObjectsTableTreeView);
+	MinorTabConfig.WorkspaceGroup = InOutExtender.GetWorkspaceGroup();
+
+	InOutExtender.GetLayoutExtender().ExtendLayout(FTimingProfilerTabs::StatsCountersID
+		, ELayoutExtensionPosition::After
+		, FTabManager::FTab(FObjectProfilerTabs::ObjectsTableTreeViewTabId, ETabState::ClosedTab));
+}
+
+////////////////////////////////////////////////////////////////////////////////////////////////////
+
+TSharedRef<SDockTab> FObjectProfilerManager::SpawnTab_ObjectsTableTreeView(const FSpawnTabArgs& Args)
+{
+	const TSharedRef<SDockTab> DockTab = SNew(SDockTab)
+		.ShouldAutosize(false)
+		.TabRole(ETabRole::PanelTab)
+		[
+			SAssignNew(ObjectsTableTreeView, SObjectTableTreeView)
+		];
+
+	DockTab->SetOnTabClosed(SDockTab::FOnTabClosedCallback::CreateRaw(this, &FObjectProfilerManager::OnObjectsTableTreeViewTabClosed));
+
+	return DockTab;
+}
+
+////////////////////////////////////////////////////////////////////////////////////////////////////
+
+bool FObjectProfilerManager::CanSpawnTab_ObjectsTableTreeView(const FSpawnTabArgs& Args)
+{
+	return bIsAvailable;
+}
+
+////////////////////////////////////////////////////////////////////////////////////////////////////
+
+void FObjectProfilerManager::OnObjectsTableTreeViewTabClosed(TSharedRef<SDockTab> TabBeingClosed)
+{
+	if (ObjectsTableTreeView.IsValid())
+	{
+		ObjectsTableTreeView->OnClose();
+	}
+
+	ObjectsTableTreeView.Reset();
+}
+
+////////////////////////////////////////////////////////////////////////////////////////////////////
+
+void FObjectProfilerManager::OnWindowClosedEvent()
+{
+	if (ObjectsTableTreeView.IsValid())
+	{
+		ObjectsTableTreeView->OnClose();
+	}
+
+	TSharedPtr<FTabManager> TimingTabManagerSharedPtr = TimingTabManager.Pin();
+
+	if (TimingTabManagerSharedPtr.IsValid())
+	{
+		TSharedPtr<SDockTab> Tab = TimingTabManagerSharedPtr->FindExistingLiveTab(FObjectProfilerTabs::ObjectsTableTreeViewTabId);
+		if (Tab.IsValid())
+		{
+			Tab->RequestCloseTab();
+			Tab->RemoveTabFromParent();
+		}
+	}
+}
+
+////////////////////////////////////////////////////////////////////////////////////////////////////
+
+} // namespace UE::Insights::ObjectProfiler
+
+#undef LOCTEXT_NAMESPACE
diff --git a/Engine/Source/Developer/TraceInsights/Private/Insights/ObjectProfiler/ObjectProfilerManager.h b/Engine/Source/Developer/TraceInsights/Private/Insights/ObjectProfiler/ObjectProfilerManager.h
new file mode 100644
--- /dev/null
+++ b/Engine/Source/Developer/TraceInsights/Private/Insights/ObjectProfiler/ObjectProfilerManager.h
@@ -0,0 +1,102 @@
+// Copyright Epic Games, Inc. All Rights Reserved.
+
+#pragma once
+
+#include "CoreTypes.h"
+
+#include "Containers/Ticker.h"
+#include "Framework/Commands/UICommandList.h"
+#include "Templates/SharedPointer.h"
+#include "Logging/LogMacros.h"
+
+// TraceInsightsCore
+#include "InsightsCore/Common/AvailabilityCheck.h"
+
+// TraceInsights
+#include "Insights/InsightsManager.h"
+#include "Insights/IUnrealInsightsModule.h"
+
+namespace UE::Insights::ObjectProfiler
+{
+
+DECLARE_LOG_CATEGORY_EXTERN(LogObjectProfiler, Log, All);
+
+class SObjectTableTreeView;
+
+struct FObjectProfilerTabs
+{
+	// Tab identifiers
+	static const FName ObjectsTableTreeViewTabId;
+};
+
+////////////////////////////////////////////////////////////////////////////////////////////////////
+/**
+ * This class manages the Cooking Profiler state and settings.
+ */
+class FObjectProfilerManager : public TSharedFromThis<FObjectProfilerManager>, public IInsightsComponent
+{
+public:
+	/** Creates the Cook Profiler manager, only one instance can exist. */
+	FObjectProfilerManager(TSharedRef<FUICommandList> InCommandList);
+
+	/** Virtual destructor. */
+	virtual ~FObjectProfilerManager();
+
+	/** Creates an instance of the CookProfiler manager. */
+	static TSharedPtr<FObjectProfilerManager> CreateInstance();
+
+	/**
+	 * @return the global instance of the Cook Profiler manager.
+	 * This is an internal singleton and cannot be used outside TraceInsights.
+	 * For external use:
+	 *     IUnrealInsightsModule& Module = FModuleManager::Get().LoadModuleChecked<IUnrealInsightsModule>("TraceInsights");
+	 *     Module.GetCookProfilerManager();
+	 */
+	static TSharedPtr<FObjectProfilerManager> Get();
+
+	//////////////////////////////////////////////////
+	// IInsightsComponent
+
+	virtual void Initialize(IUnrealInsightsModule& InsightsModule) override;
+	virtual void Shutdown() override;
+	virtual void RegisterMajorTabs(IUnrealInsightsModule& InsightsModule) override;
+	virtual void UnregisterMajorTabs() override;
+	virtual void OnWindowClosedEvent() override;
+
+	//////////////////////////////////////////////////
+
+	TSharedRef<SDockTab> SpawnTab_ObjectsTableTreeView(const FSpawnTabArgs& Args);
+	bool CanSpawnTab_ObjectsTableTreeView(const FSpawnTabArgs& Args);
+	void OnObjectsTableTreeViewTabClosed(TSharedRef<SDockTab> TabBeingClosed);
+
+	//////////////////////////////////////////////////
+
+	bool IsAvailable() { return bIsAvailable; }
+
+	void OnSessionChanged();
+
+private:
+	bool Tick(float DeltaTime);
+
+	void RegisterTimingProfilerLayoutExtensions(FInsightsMajorTabExtender& InOutExtender);
+
+private:
+	bool bIsInitialized;
+	bool bIsAvailable;
+	FAvailabilityCheck AvailabilityCheck;
+
+	/** The delegate to be invoked when this manager ticks. */
+	FTickerDelegate OnTick;
+
+	/** Handle to the registered OnTick. */
+	FTSTicker::FDelegateHandle OnTickHandle;
+
+	/** A shared pointer to the global instance of the Cook Profiler manager. */
+	static TSharedPtr<FObjectProfilerManager> Instance;
+
+	TWeakPtr<FTabManager> TimingTabManager;
+
+	TSharedPtr<SObjectTableTreeView> ObjectsTableTreeView;
+};
+
+} // namespace UE::Insights::ObjectProfiler
diff --git a/Engine/Source/Developer/TraceInsights/Private/Insights/ObjectProfiler/ViewModels/ObjectGroupingByOuter.cpp b/Engine/Source/Developer/TraceInsights/Private/Insights/ObjectProfiler/ViewModels/ObjectGroupingByOuter.cpp
new file mode 100644
--- /dev/null
+++ b/Engine/Source/Developer/TraceInsights/Private/Insights/ObjectProfiler/ViewModels/ObjectGroupingByOuter.cpp
@@ -0,0 +1,156 @@
+// Copyright Epic Games, Inc. All Rights Reserved.
+
+#include "ObjectGroupingByOuter.h"
+
+#include "Containers/BitArray.h"
+#include "Internationalization/Internationalization.h"
+
+// TraceInsightsCore
+#include "InsightsCore/Common/AsyncOperationProgress.h"
+
+// TraceInsights
+#include "Insights/ObjectProfiler/ViewModels/ObjectNode.h"
+
+#define LOCTEXT_NAMESPACE "UE::Insights::MemoryProfiler::FMemAllocNode"
+
+namespace UE::Insights::ObjectProfiler
+{
+
+INSIGHTS_IMPLEMENT_RTTI(FObjectGroupingByOuter)
+
+////////////////////////////////////////////////////////////////////////////////////////////////////
+
+FObjectGroupingByOuter::FObjectGroupingByOuter()
+	: FTreeNodeGrouping(
+		LOCTEXT("Grouping_ByOuter_ShortName", "Outer"),
+		LOCTEXT("Grouping_ByOuter_TitleName", "By Outer"),
+		LOCTEXT("Grouping_ByOuter_Desc", "Creates a tree based on object hierarchy."),
+		nullptr)
+{
+	SetColor(FLinearColor(1.0f, 0.45f, 0.6f, 1.0f));
+}
+
+////////////////////////////////////////////////////////////////////////////////////////////////////
+
+FObjectGroupingByOuter::~FObjectGroupingByOuter()
+{
+}
+
+////////////////////////////////////////////////////////////////////////////////////////////////////
+
+void FObjectGroupingByOuter::GroupNodes(
+	const TArray<FTableTreeNodePtr>& Nodes,
+	FTableTreeNode& ParentGroup,
+	TWeakPtr<FTable> InParentTable,
+	IAsyncOperationProgress& InAsyncOperationProgress) const
+{
+	ParentGroup.ClearChildren();
+	ParentGroupPtr = &ParentGroup;
+
+	for (FTableTreeNodePtr NodePtr : Nodes)
+	{
+		if (InAsyncOperationProgress.ShouldCancelAsyncOp())
+		{
+			break;
+		}
+
+		if (!NodePtr->Is<FObjectNode>())
+		{
+			ParentGroup.AddChildAndSetParent(NodePtr);
+			continue;
+		}
+
+		FObjectNode& ObjNode = NodePtr->As<FObjectNode>();
+		if (ObjNode.GetObjectId() == uint32(-1))
+		{
+			ParentGroup.AddChildAndSetParent(NodePtr);
+			continue;
+		}
+
+		TSharedPtr<FObjectNode> NewNode = AddObject(ObjNode);
+		NewNode->ResetIsMissing();
+	}
+
+	ObjectNodes.Reset();
+}
+
+////////////////////////////////////////////////////////////////////////////////////////////////////
+
+TSharedPtr<FObjectNode> FObjectGroupingByOuter::AddObject(FObjectNode& ObjNode) const
+{
+	const uint32 ObjId = ObjNode.GetObjectId();
+	if (ObjId >= uint32(ObjectNodes.Num()))
+	{
+		ObjectNodes.AddDefaulted(ObjId - ObjectNodes.Num() + 1);
+	}
+	if (ObjectNodes[ObjId].IsValid())
+	{
+		return ObjectNodes[ObjId];
+	}
+
+	TSharedPtr<FObjectNode> NewNode = MakeShared<FObjectNode>(ObjNode);
+	NewNode->SetIsMissing();
+	ObjectNodes[ObjId] = NewNode;
+
+	if (bCreateChildrenSubGroups)
+	{
+		TSharedPtr<FObjectNode> Outer = ObjNode.GetOuter();
+		if (!Outer.IsValid())
+		{
+			// Create the Object group node (root object)
+			auto ObjectGroupNode = CreateObjectGroupNode(ObjNode);
+			ObjectGroupNode->AddChildAndSetParent(NewNode);
+			ParentGroupPtr->AddChildAndSetParent(ObjectGroupNode);
+			return NewNode;
+		}
+
+		// Get or create the "Children" group node
+		TSharedPtr<FObjectNode> OuterClone = AddObject(*Outer);
+		auto OuterParentGroup = OuterClone->GetParent();
+		check(OuterParentGroup.IsValid() && OuterParentGroup->Is<FCustomTableTreeNode>());
+		if (OuterParentGroup->GetChildrenCount() == 1)
+		{
+			auto ChildrenGroupNode = MakeShared<FCustomTableTreeNode>(ObjNode.GetParentTable(), LOCTEXT("GroupName_Children", "Children"));
+			OuterParentGroup->AddChildAndSetParent(ChildrenGroupNode);
+		}
+		check(OuterParentGroup->GetChildrenCount() > 1);
+		auto ChildrenGroupNode = OuterParentGroup->GetChildren()[1];
+		check(ChildrenGroupNode->Is<FCustomTableTreeNode>());
+
+		// Create the Object group node (child object)
+		auto ObjectGroupNode = CreateObjectGroupNode(ObjNode);
+		ObjectGroupNode->AddChildAndSetParent(NewNode);
+		ChildrenGroupNode->AddChildAndSetParent(ObjectGroupNode);
+		return NewNode;
+	}
+	else
+	{
+		TSharedPtr<FObjectNode> Outer = ObjNode.GetOuter();
+		if (!Outer.IsValid())
+		{
+			ParentGroupPtr->AddChildAndSetParent(NewNode);
+			return NewNode;
+		}
+
+		TSharedPtr<FObjectNode> OuterClone = AddObject(*Outer);
+		OuterClone->InitGroupData();
+		OuterClone->AddChildAndSetParent(NewNode);
+		return NewNode;
+	}
+}
+
+////////////////////////////////////////////////////////////////////////////////////////////////////
+
+TSharedPtr<FCustomTableTreeNode> FObjectGroupingByOuter::CreateObjectGroupNode(FObjectNode& ObjNode) const
+{
+	const FText DisplayName = FText::FromString(ObjNode.GetObjectName());
+	auto ObjectGroupNode = MakeShared<FCustomTableTreeNode>(ObjNode.GetParentTable(), DisplayName);
+	ObjectGroupNode->SetColor(FLinearColor(0.9f, 0.9f, 1.0f, 1.0f));
+	return ObjectGroupNode;
+}
+
+////////////////////////////////////////////////////////////////////////////////////////////////////
+
+} // namespace UE::Insights::ObjectProfiler
+
+#undef LOCTEXT_NAMESPACE
diff --git a/Engine/Source/Developer/TraceInsights/Private/Insights/ObjectProfiler/ViewModels/ObjectGroupingByOuter.h b/Engine/Source/Developer/TraceInsights/Private/Insights/ObjectProfiler/ViewModels/ObjectGroupingByOuter.h
new file mode 100644
--- /dev/null
+++ b/Engine/Source/Developer/TraceInsights/Private/Insights/ObjectProfiler/ViewModels/ObjectGroupingByOuter.h
@@ -0,0 +1,36 @@
+// Copyright Epic Games, Inc. All Rights Reserved.
+
+#pragma once
+
+#include "CoreTypes.h"
+
+// TraceInsightsCore
+#include "InsightsCore/Table/ViewModels/BaseTreeNode.h"
+#include "InsightsCore/Table/ViewModels/TableTreeNode.h"
+#include "InsightsCore/Table/ViewModels/TreeNodeGrouping.h"
+
+namespace UE::Insights::ObjectProfiler
+{
+
+class FObjectNode;
+
+class FObjectGroupingByOuter : public FTreeNodeGrouping
+{
+	INSIGHTS_DECLARE_RTTI(FObjectGroupingByOuter, FTreeNodeGrouping)
+
+public:
+	FObjectGroupingByOuter();
+	virtual ~FObjectGroupingByOuter() override;
+
+	virtual void GroupNodes(const TArray<FTableTreeNodePtr>& Nodes, FTableTreeNode& ParentGroup, TWeakPtr<FTable> InParentTable, IAsyncOperationProgress& InAsyncOperationProgress) const override;
+
+private:
+	TSharedPtr<FObjectNode> AddObject(FObjectNode& ObjNode) const;
+	TSharedPtr<FCustomTableTreeNode> CreateObjectGroupNode(FObjectNode& ObjNode) const;
+
+	bool bCreateChildrenSubGroups = true;
+	mutable FTableTreeNode* ParentGroupPtr = nullptr;
+	mutable TArray<TSharedPtr<FObjectNode>> ObjectNodes;
+};
+
+} // namespace UE::Insights::ObjectProfiler
diff --git a/Engine/Source/Developer/TraceInsights/Private/Insights/ObjectProfiler/ViewModels/ObjectNode.cpp b/Engine/Source/Developer/TraceInsights/Private/Insights/ObjectProfiler/ViewModels/ObjectNode.cpp
new file mode 100644
--- /dev/null
+++ b/Engine/Source/Developer/TraceInsights/Private/Insights/ObjectProfiler/ViewModels/ObjectNode.cpp
@@ -0,0 +1,149 @@
+// Copyright Epic Games, Inc. All Rights Reserved.
+
+#include "ObjectNode.h"
+
+// TraceInsights
+#include "Insights/InsightsStyle.h"
+#include "Insights/ObjectProfiler/ObjectProfilerManager.h"
+#include "Insights/ObjectProfiler/ViewModels/ObjectTable.h"
+
+#define LOCTEXT_NAMESPACE "UE::Insights::ObjectProfiler::FObjectNode"
+
+namespace UE::Insights::ObjectProfiler
+{
+
+INSIGHTS_IMPLEMENT_RTTI(FObjectNode)
+
+////////////////////////////////////////////////////////////////////////////////////////////////////
+
+FObjectNode::FObjectNode(TWeakPtr<FObjectTable> InParentTable, const TraceServices::FObjectInfo& ObjectInfo)
+	: FTableTreeNode(InParentTable, int32(ObjectInfo.Id))
+	, ObjectId(ObjectInfo.Id)
+	, ClassId(ObjectInfo.ClassId)
+	, OuterId(ObjectInfo.OuterId)
+	, PackageId(ObjectInfo.PackageId)
+	, ObjectName(ObjectInfo.Name)
+	, VersePath(ObjectInfo.VersePath)
+	, ObjectFlags(ObjectInfo.Flags)
+	, StructureSize(ObjectInfo.StructureSize)
+	, SystemMemorySize(int64(ObjectInfo.SystemMemoryBytes))
+	, VideoMemorySize(int64(ObjectInfo.VideoMemoryBytes))
+{
+}
+
+////////////////////////////////////////////////////////////////////////////////////////////////////
+
+FObjectNode::FObjectNode(const FObjectNode& Clone)
+	: FTableTreeNode(Clone.GetParentTable(), int32(Clone.ObjectId))
+	, ObjectId(Clone.ObjectId)
+	, ClassId(Clone.ClassId)
+	, OuterId(Clone.OuterId)
+	, PackageId(Clone.PackageId)
+	, Class(Clone.Class)
+	, Outer(Clone.Outer)
+	, Package(Clone.Package)
+	, ObjectName(Clone.ObjectName)
+	, VersePath(Clone.VersePath)
+	, ObjectPath(Clone.ObjectPath)
+	, ObjectFlags(Clone.ObjectFlags)
+	, StructureSize(Clone.StructureSize)
+	, SystemMemorySize(Clone.SystemMemorySize)
+	, VideoMemorySize(Clone.VideoMemorySize)
+	, InternalFlags(Clone.InternalFlags)
+{
+}
+
+////////////////////////////////////////////////////////////////////////////////////////////////////
+
+FObjectNode::FObjectNode(TWeakPtr<FObjectTable> InParentTable)
+	: FTableTreeNode(InParentTable)
+{
+}
+
+////////////////////////////////////////////////////////////////////////////////////////////////////
+
+const FText FObjectNode::GetDisplayName() const
+{
+	return FText::FromString(ObjectName);
+}
+
+////////////////////////////////////////////////////////////////////////////////////////////////////
+
+const FText FObjectNode::GetExtraDisplayName() const
+{
+	return FText::GetEmpty();
+}
+
+////////////////////////////////////////////////////////////////////////////////////////////////////
+
+bool FObjectNode::HasExtraDisplayName() const
+{
+	return false;
+}
+
+////////////////////////////////////////////////////////////////////////////////////////////////////
+
+const FText FObjectNode::GetTooltipText() const
+{
+	return FText::GetEmpty();
+}
+
+////////////////////////////////////////////////////////////////////////////////////////////////////
+
+const FSlateBrush* FObjectNode::GetIcon() const
+{
+	if (IsClass())
+	{
+		return FInsightsStyle::GetBrush("Icons.UClass.TreeItem");
+	}
+	else if (IsPackage())
+	{
+		return FInsightsStyle::GetBrush("Icons.UPackage.TreeItem");
+	}
+	else
+	{
+		return FInsightsStyle::GetBrush("Icons.UObject.TreeItem");
+	}
+}
+
+////////////////////////////////////////////////////////////////////////////////////////////////////
+
+FLinearColor FObjectNode::GetIconColor() const
+{
+	if (IsClass())
+	{
+		return FLinearColor(0.5f, 1.0f, 1.0f, 1.0f);
+	}
+	else if (IsPackage())
+	{
+		return FLinearColor(1.0f, 1.0f, 0.5f, 1.0f);
+	}
+	else
+	{
+		return GetDefaultColor(IsGroup());
+	}
+}
+
+////////////////////////////////////////////////////////////////////////////////////////////////////
+
+FLinearColor FObjectNode::GetColor() const
+{
+	if (ObjectId == uint32(-1))
+	{
+		return GetDefaultIconColor(true);
+	}
+	else if (IsMissing())
+	{
+		return FLinearColor(1.0f, 0.7f, 0.7f, 1.0f);
+	}
+	else
+	{
+		return GetDefaultIconColor(false);
+	}
+}
+
+////////////////////////////////////////////////////////////////////////////////////////////////////
+
+} // namespace UE::Insights::ObjectProfiler
+
+#undef LOCTEXT_NAMESPACE
diff --git a/Engine/Source/Developer/TraceInsights/Private/Insights/ObjectProfiler/ViewModels/ObjectNode.h b/Engine/Source/Developer/TraceInsights/Private/Insights/ObjectProfiler/ViewModels/ObjectNode.h
new file mode 100644
--- /dev/null
+++ b/Engine/Source/Developer/TraceInsights/Private/Insights/ObjectProfiler/ViewModels/ObjectNode.h
@@ -0,0 +1,135 @@
+// Copyright Epic Games, Inc. All Rights Reserved.
+
+#pragma once
+
+#include "CoreTypes.h"
+
+// TraceServices
+#include "TraceServices/Model/ObjectProvider.h"
+
+// TraceInsightsCore
+#include "InsightsCore/Table/ViewModels/TableTreeNode.h"
+#include "InsightsCore/Table/ViewModels/TreeNodeGrouping.h"
+
+// TraceInsights
+
+namespace TraceServices
+{
+	struct FObjectInfo;
+}
+
+namespace UE::Insights::ObjectProfiler
+{
+
+class FObjectTable;
+class SObjectTableTreeView;
+
+////////////////////////////////////////////////////////////////////////////////////////////////////
+/**
+ * Class used to store information about a UObject node (used in the SObjectTableTreeView).
+ */
+class FObjectNode : public FTableTreeNode
+{
+	friend class SObjectTableTreeView;
+	friend class FObjectGroupingByOuter;
+
+	INSIGHTS_DECLARE_RTTI(FObjectNode, FTableTreeNode)
+
+private:
+	enum class EInternalFlags
+	{
+		None = 0,
+		IsClass = 1 << 0,
+		IsPackage = 1 << 1,
+		IsMissing = 1 << 2, // missing intermediate node in the Outer hierarchy
+	};
+
+public:
+	/** Initialization constructor for the UObject node. */
+	explicit FObjectNode(TWeakPtr<FObjectTable> InParentTable, const TraceServices::FObjectInfo& ObjectInfo);
+
+	explicit FObjectNode(const FObjectNode& Clone);
+
+	/** Initialization constructor for the group node. */
+	explicit FObjectNode(TWeakPtr<FObjectTable> InParentTable);
+
+	virtual const FText GetDisplayName() const override;
+	virtual const FText GetExtraDisplayName() const override;
+	virtual bool HasExtraDisplayName() const override;
+	virtual const FText GetTooltipText() const override;
+
+	virtual const FSlateBrush* GetIcon() const override;
+	virtual FLinearColor GetIconColor() const override;
+	virtual FLinearColor GetColor() const override;
+
+	uint32 GetObjectId() const { return ObjectId; }
+
+	uint32 GetClassId() const { return ClassId; }
+	uint32 GetOuterId() const { return OuterId; }
+	uint32 GetPackageId() const { return PackageId; }
+
+	TSharedPtr<FObjectNode> GetClass() const { return Class.Pin(); }
+	TSharedPtr<FObjectNode> GetOuter() const { return Outer.Pin(); }
+	TSharedPtr<FObjectNode> GetPackage() const { return Package.Pin(); }
+
+	const TCHAR* GetObjectName() const { return ObjectName; }
+	const TCHAR* GetObjectPath() const { return *ObjectPath; }
+	const TCHAR* GetVersePath() const { return VersePath; }
+
+	uint32 GetObjectFlags() const { return ObjectFlags; }
+
+	int32 GetStructureSize() const { return StructureSize; }
+
+	int64 GetSystemMemorySize() const { return SystemMemorySize; }
+	int64 GetVideoMemorySize() const { return VideoMemorySize; }
+	int64 GetEstimatedSize() const { return SystemMemorySize + VideoMemorySize; }
+
+	int64 GetInclusiveSystemMemorySize() const { return InclSystemMemorySize; }
+	int64 GetInclusiveVideoMemorySize() const { return InclVideoMemorySize; }
+	int64 GetInclusiveEstimatedSize() const { return InclSystemMemorySize + InclVideoMemorySize; }
+
+	uint32 GetNumReferences() const { return 0; } // TODO
+
+	bool IsClass() const { return (uint32(InternalFlags) & uint32(EInternalFlags::IsClass)) != 0; }
+	bool IsPackage() const { return (uint32(InternalFlags) & uint32(EInternalFlags::IsPackage)) != 0; }
+	bool IsMissing() const { return (uint32(InternalFlags) & uint32(EInternalFlags::IsMissing)) != 0; }
+
+private:
+	void SetIsClass() { InternalFlags = EInternalFlags(uint32(InternalFlags) | uint32(EInternalFlags::IsClass)); }
+	void SetIsPackage() { InternalFlags = EInternalFlags(uint32(InternalFlags) | uint32(EInternalFlags::IsPackage)); }
+	void SetIsMissing() { InternalFlags = EInternalFlags(uint32(InternalFlags) | uint32(EInternalFlags::IsMissing)); }
+	void ResetIsMissing() { InternalFlags = EInternalFlags(uint32(InternalFlags) & ~uint32(EInternalFlags::IsMissing)); }
+
+	void SetClass(TWeakPtr<FObjectNode> InClass) { Class = InClass; }
+	void SetOuter(TWeakPtr<FObjectNode> InOuter) { Outer = InOuter; }
+	void SetPackage(TWeakPtr<FObjectNode> InPackage) { Package = InPackage; }
+
+private:
+	uint32 ObjectId = uint32(-1);
+
+	uint32 ClassId = uint32(-1);
+	uint32 OuterId = uint32(-1);
+	uint32 PackageId = uint32(-1);
+
+	TWeakPtr<FObjectNode> Class;
+	TWeakPtr<FObjectNode> Outer;
+	TWeakPtr<FObjectNode> Package;
+
+	const TCHAR* ObjectName = nullptr;
+	const TCHAR* VersePath = nullptr;
+	FString ObjectPath;
+
+	uint32 ObjectFlags = 0;
+
+	int32 StructureSize = 0; // for UStruct and UClass
+	int64 SystemMemorySize = 0;
+	int64 VideoMemorySize = 0;
+	int64 InclSystemMemorySize = 0;
+	int64 InclVideoMemorySize = 0;
+
+	EInternalFlags InternalFlags = EInternalFlags::None;
+};
+
+////////////////////////////////////////////////////////////////////////////////////////////////////
+
+} // namespace UE::Insights::ObjectProfiler
diff --git a/Engine/Source/Developer/TraceInsights/Private/Insights/ObjectProfiler/ViewModels/ObjectTable.cpp b/Engine/Source/Developer/TraceInsights/Private/Insights/ObjectProfiler/ViewModels/ObjectTable.cpp
new file mode 100644
--- /dev/null
+++ b/Engine/Source/Developer/TraceInsights/Private/Insights/ObjectProfiler/ViewModels/ObjectTable.cpp
@@ -0,0 +1,1133 @@
+// Copyright Epic Games, Inc. All Rights Reserved.
+
+#include "ObjectTable.h"
+
+// TraceInsightsCore
+#include "InsightsCore/Table/ViewModels/TableCellValueFormatter.h"
+#include "InsightsCore/Table/ViewModels/TableCellValueGetter.h"
+#include "InsightsCore/Table/ViewModels/TableCellValueSorter.h"
+#include "InsightsCore/Table/ViewModels/TableColumn.h"
+
+// TraceInsights
+#include "Insights/ObjectProfiler/ViewModels/ObjectNode.h"
+
+#define LOCTEXT_NAMESPACE "UE::Insights::ObjectProfiler::FObjectTable"
+
+namespace UE::Insights::ObjectProfiler
+{
+
+////////////////////////////////////////////////////////////////////////////////////////////////////
+// Column identifiers
+
+const FName FObjectTableColumns::ObjectCountColumnId("ObjectCount");
+const FName FObjectTableColumns::ObjectNameColumnId("ObjectName");
+const FName FObjectTableColumns::ObjectIdColumnId("ObjectId");
+const FName FObjectTableColumns::ClassNameColumnId("ClassName");
+const FName FObjectTableColumns::ClassIdColumnId("ClassId");
+const FName FObjectTableColumns::OuterNameColumnId("OuterName");
+const FName FObjectTableColumns::OuterIdColumnId("OuterId");
+const FName FObjectTableColumns::PackageNameColumnId("PackageName");
+const FName FObjectTableColumns::PackageIdColumnId("PackageId");
+const FName FObjectTableColumns::ObjectFlagsColumnId("ObjectFlags");
+const FName FObjectTableColumns::ObjectPathColumnId("ObjectPath");
+const FName FObjectTableColumns::VersePathColumnId("VersePath");
+const FName FObjectTableColumns::StructureSizeColumnId("StructureSize");
+const FName FObjectTableColumns::SystemMemorySizeColumnId("SystemMemorySize");
+const FName FObjectTableColumns::VideoMemorySizeColumnId("VideoMemorySize");
+const FName FObjectTableColumns::EstimatedSizeColumnId("EstimatedSize");
+const FName FObjectTableColumns::ImpactColumnId("Impact");
+const FName FObjectTableColumns::ReferencesColumnId("References");
+
+////////////////////////////////////////////////////////////////////////////////////////////////////
+// FObjectTable
+////////////////////////////////////////////////////////////////////////////////////////////////////
+
+FObjectTable::FObjectTable()
+{
+}
+
+////////////////////////////////////////////////////////////////////////////////////////////////////
+
+FObjectTable::~FObjectTable()
+{
+}
+
+////////////////////////////////////////////////////////////////////////////////////////////////////
+
+void FObjectTable::Reset()
+{
+	FTable::Reset();
+
+	AddDefaultColumns();
+}
+
+////////////////////////////////////////////////////////////////////////////////////////////////////
+
+void FObjectTable::AddDefaultColumns()
+{
+	class FAggregatedValueGetter : public FTableCellValueGetter
+	{
+	public:
+		virtual const TOptional<FTableCellValue> GetNodeValue(const FBaseTreeNode& Node) const
+		{
+			return TOptional<FTableCellValue>();
+		}
+
+		virtual const TOptional<FTableCellValue> GetValue(const FTableColumn& Column, const FBaseTreeNode& Node) const override
+		{
+			//if (Column.GetId() == FObjectTableColumns::ClassNameColumnId &&
+			//	Node.Is<FObjectNode>() &&
+			//	Node.As<FObjectNode>().GetObjectId() != uint32(-1))
+			//{
+			//	return GetNodeValue(Node);
+			//}
+			if (Node.IsGroup() && Node.Is<FTableTreeNode>())
+			{
+				const FTableTreeNode& TableTreeNode = Node.As<FTableTreeNode>();
+				if (TableTreeNode.HasAggregatedValue(Column.GetId()))
+				{
+					return TableTreeNode.GetAggregatedValue(Column.GetId());
+				}
+			}
+			return GetNodeValue(Node);
+		}
+	};
+
+	//////////////////////////////////////////////////
+	// Hierarchy Column
+	{
+		const int32 HierarchyColumnIndex = -1;
+		const TCHAR* HierarchyColumnName = nullptr;
+		AddHierarchyColumn(HierarchyColumnIndex, HierarchyColumnName);
+
+		FTableColumn& HierarchyColumn = *GetColumns()[0];
+		HierarchyColumn.SetInitialWidth(200.0f);
+		HierarchyColumn.SetShortName(LOCTEXT("HierarchyColumnName", "Hierarchy"));
+		HierarchyColumn.SetTitleName(LOCTEXT("HierarchyColumnTitle", "UObject Hierarchy"));
+		HierarchyColumn.SetDescription(LOCTEXT("HierarchyColumnDesc", "Hierarchy of the UObject's tree"));
+	}
+
+	int32 ColumnIndex = 0;
+
+	//////////////////////////////////////////////////
+	// Object Count Column
+	{
+		TSharedRef<FTableColumn> ColumnRef = MakeShared<FTableColumn>(FObjectTableColumns::ObjectCountColumnId);
+		FTableColumn& Column = *ColumnRef;
+
+		Column.SetIndex(ColumnIndex++);
+
+		Column.SetShortName(LOCTEXT("ObjectCountColumnName", "Count"));
+		Column.SetTitleName(LOCTEXT("ObjectCountColumnTitle", "Object Count"));
+		Column.SetDescription(LOCTEXT("ObjectCountColumnDesc", "The number of the UObject instances"));
+
+		Column.SetFlags(ETableColumnFlags::CanBeHidden |
+						ETableColumnFlags::CanBeFiltered);
+
+		Column.SetHorizontalAlignment(HAlign_Right);
+		Column.SetInitialWidth(60.0f);
+
+		Column.SetDataType(ETableCellDataType::Int64);
+
+		class FValueGetter : public FAggregatedValueGetter
+		{
+		public:
+			virtual const TOptional<FTableCellValue> GetNodeValue(const FBaseTreeNode& Node) const override
+			{
+				if (Node.Is<FObjectNode>())
+				{
+					return FTableCellValue(static_cast<int64>(1));
+				}
+				return TOptional<FTableCellValue>();
+			}
+		};
+		TSharedRef<ITableCellValueGetter> Getter = MakeShared<FValueGetter>();
+		Column.SetValueGetter(Getter);
+
+		TSharedRef<ITableCellValueFormatter> Formatter = MakeShared<FInt64ValueFormatterAsNumber>();
+		Column.SetValueFormatter(Formatter);
+
+		TSharedRef<ITableCellValueSorter> Sorter = MakeShared<FSorterByInt64Value>(ColumnRef);
+		Column.SetValueSorter(Sorter);
+		Column.SetInitialSortMode(EColumnSortMode::Descending);
+
+		Column.SetAggregation(ETableColumnAggregation::Sum);
+
+		AddColumn(ColumnRef);
+	}
+	//////////////////////////////////////////////////
+	// Object Name Column
+	{
+		TSharedRef<FTableColumn> ColumnRef = MakeShared<FTableColumn>(FObjectTableColumns::ObjectNameColumnId);
+		FTableColumn& Column = *ColumnRef;
+
+		Column.SetIndex(ColumnIndex++);
+
+		Column.SetShortName(LOCTEXT("ObjectNameColumnName", "Name"));
+		Column.SetTitleName(LOCTEXT("ObjectNameColumnTitle", "Object Name"));
+		Column.SetDescription(LOCTEXT("ObjectNameColumnDesc", "The name of the UObject instance"));
+
+		Column.SetFlags(ETableColumnFlags::CanBeHidden |
+						ETableColumnFlags::CanBeFiltered);
+
+		Column.SetHorizontalAlignment(HAlign_Left);
+		Column.SetInitialWidth(120.0f);
+
+		Column.SetDataType(ETableCellDataType::CString);
+
+		class FValueGetter : public FAggregatedValueGetter
+		{
+		public:
+			virtual const TOptional<FTableCellValue> GetNodeValue(const FBaseTreeNode& Node) const override
+			{
+				if (Node.Is<FObjectNode>())
+				{
+					const FObjectNode& ObjectNode = Node.As<FObjectNode>();
+					return FTableCellValue(ObjectNode.GetObjectName());
+				}
+				return TOptional<FTableCellValue>();
+			}
+		};
+		TSharedRef<ITableCellValueGetter> Getter = MakeShared<FValueGetter>();
+		Column.SetValueGetter(Getter);
+
+		TSharedRef<ITableCellValueFormatter> Formatter = MakeShared<FCStringValueFormatterAsText>();
+		Column.SetValueFormatter(Formatter);
+
+		TSharedRef<ITableCellValueSorter> Sorter = MakeShared<FSorterByCStringValue>(ColumnRef);
+		Column.SetValueSorter(Sorter);
+
+		Column.SetAggregation(ETableColumnAggregation::SameValue);
+
+		AddColumn(ColumnRef);
+
+		// Set the same "by name" sorter also for the hierarchy column.
+		//FTableColumn& HierarchyColumn = *GetColumns()[0];
+		//HierarchyColumn.SetValueSorter(Sorter);
+	}
+	//////////////////////////////////////////////////
+	// Object Id Column
+	{
+		TSharedRef<FTableColumn> ColumnRef = MakeShared<FTableColumn>(FObjectTableColumns::ObjectIdColumnId);
+		FTableColumn& Column = *ColumnRef;
+
+		Column.SetIndex(ColumnIndex++);
+
+		Column.SetShortName(LOCTEXT("ObjectIdColumnName", "Object Id"));
+		Column.SetTitleName(LOCTEXT("ObjectIdColumnTitle", "Object Id"));
+		Column.SetDescription(LOCTEXT("ObjectIdColumnDesc", "The id of the UObject instance"));
+
+		Column.SetFlags(ETableColumnFlags::ShouldBeVisible |
+						ETableColumnFlags::CanBeHidden |
+						ETableColumnFlags::CanBeFiltered);
+
+		Column.SetHorizontalAlignment(HAlign_Left);
+		Column.SetInitialWidth(60.0f);
+
+		Column.SetDataType(ETableCellDataType::Int64);
+
+		class FValueGetter : public FAggregatedValueGetter
+		{
+		public:
+			virtual const TOptional<FTableCellValue> GetNodeValue(const FBaseTreeNode& Node) const override
+			{
+				if (Node.Is<FObjectNode>())
+				{
+					const FObjectNode& ObjectNode = Node.As<FObjectNode>();
+					const uint32 Id = ObjectNode.GetObjectId();
+					if (Id != uint32(-1))
+					{
+						return FTableCellValue(int64(Id));
+					}
+				}
+				return TOptional<FTableCellValue>();
+			}
+		};
+		TSharedRef<ITableCellValueGetter> Getter = MakeShared<FValueGetter>();
+		Column.SetValueGetter(Getter);
+
+		TSharedRef<ITableCellValueFormatter> Formatter = MakeShared<FInt64ValueFormatterAsNumber>();
+		Column.SetValueFormatter(Formatter);
+
+		TSharedRef<ITableCellValueSorter> Sorter = MakeShared<FSorterByInt64Value>(ColumnRef);
+		Column.SetValueSorter(Sorter);
+		Column.SetInitialSortMode(EColumnSortMode::Descending);
+
+		//TSharedRef<IFilterValueConverter> Converter = MakeShared<FMemoryFilterValueConverter>();
+		//Column.SetValueConverter(Converter);
+
+		//Column.SetAggregation(ETableColumnAggregation::None);
+
+		AddColumn(ColumnRef);
+	}
+	//////////////////////////////////////////////////
+	// Class Name Column
+	{
+		TSharedRef<FTableColumn> ColumnRef = MakeShared<FTableColumn>(FObjectTableColumns::ClassNameColumnId);
+		FTableColumn& Column = *ColumnRef;
+
+		Column.SetIndex(ColumnIndex++);
+
+		Column.SetShortName(LOCTEXT("ClassNameColumnName", "Class"));
+		Column.SetTitleName(LOCTEXT("ClassNameColumnTitle", "Class Name"));
+		Column.SetDescription(LOCTEXT("ClassNameColumnDesc", "The name of the UClass of the current UObject instance"));
+
+		Column.SetFlags(ETableColumnFlags::ShouldBeVisible |
+						ETableColumnFlags::CanBeHidden |
+						ETableColumnFlags::CanBeFiltered);
+
+		Column.SetHorizontalAlignment(HAlign_Left);
+		Column.SetInitialWidth(120.0f);
+
+		Column.SetDataType(ETableCellDataType::CString);
+
+		class FValueGetter : public FAggregatedValueGetter
+		{
+		public:
+			virtual const TOptional<FTableCellValue> GetNodeValue(const FBaseTreeNode& Node) const override
+			{
+				if (Node.Is<FObjectNode>())
+				{
+					const FObjectNode& ObjectNode = Node.As<FObjectNode>();
+					TSharedPtr<FObjectNode> ClassNode = ObjectNode.GetClass();
+					if (ClassNode.IsValid())
+					{
+						return FTableCellValue(ClassNode->GetObjectName());
+					}
+				}
+				return TOptional<FTableCellValue>();
+			}
+		};
+		TSharedRef<ITableCellValueGetter> Getter = MakeShared<FValueGetter>();
+		Column.SetValueGetter(Getter);
+
+		TSharedRef<ITableCellValueFormatter> Formatter = MakeShared<FCStringValueFormatterAsText>();
+		Column.SetValueFormatter(Formatter);
+
+		TSharedRef<ITableCellValueSorter> Sorter = MakeShared<FSorterByCStringValue>(ColumnRef);
+		Column.SetValueSorter(Sorter);
+
+		Column.SetAggregation(ETableColumnAggregation::SameValue);
+
+		AddColumn(ColumnRef);
+	}
+	//////////////////////////////////////////////////
+	// Class Id Column
+	{
+		TSharedRef<FTableColumn> ColumnRef = MakeShared<FTableColumn>(FObjectTableColumns::ClassIdColumnId);
+		FTableColumn& Column = *ColumnRef;
+
+		Column.SetIndex(ColumnIndex++);
+
+		Column.SetShortName(LOCTEXT("ClassIdColumnName", "Class Id"));
+		Column.SetTitleName(LOCTEXT("ClassIdColumnTitle", "Class Id"));
+		Column.SetDescription(LOCTEXT("ClassIdColumnDesc", "The id of the UClass of the current UObject instance"));
+
+		Column.SetFlags(ETableColumnFlags::CanBeHidden |
+						ETableColumnFlags::CanBeFiltered);
+
+		Column.SetHorizontalAlignment(HAlign_Left);
+		Column.SetInitialWidth(60.0f);
+
+		Column.SetDataType(ETableCellDataType::Int64);
+
+		class FValueGetter : public FAggregatedValueGetter
+		{
+		public:
+			virtual const TOptional<FTableCellValue> GetNodeValue(const FBaseTreeNode& Node) const override
+			{
+				if (Node.Is<FObjectNode>())
+				{
+					const FObjectNode& ObjectNode = Node.As<FObjectNode>();
+					const uint32 Id = ObjectNode.GetClassId();
+					if (Id != uint32(-1))
+					{
+						return FTableCellValue(int64(Id));
+					}
+				}
+				return TOptional<FTableCellValue>();
+			}
+		};
+		TSharedRef<ITableCellValueGetter> Getter = MakeShared<FValueGetter>();
+		Column.SetValueGetter(Getter);
+
+		TSharedRef<ITableCellValueFormatter> Formatter = MakeShared<FInt64ValueFormatterAsNumber>();
+		Column.SetValueFormatter(Formatter);
+
+		TSharedRef<ITableCellValueSorter> Sorter = MakeShared<FSorterByInt64Value>(ColumnRef);
+		Column.SetValueSorter(Sorter);
+		Column.SetInitialSortMode(EColumnSortMode::Descending);
+
+		//TSharedRef<IFilterValueConverter> Converter = MakeShared<FMemoryFilterValueConverter>();
+		//Column.SetValueConverter(Converter);
+
+		Column.SetAggregation(ETableColumnAggregation::SameValue);
+
+		AddColumn(ColumnRef);
+	}
+	//////////////////////////////////////////////////
+	// Outer Name Column
+	{
+		TSharedRef<FTableColumn> ColumnRef = MakeShared<FTableColumn>(FObjectTableColumns::OuterNameColumnId);
+		FTableColumn& Column = *ColumnRef;
+
+		Column.SetIndex(ColumnIndex++);
+
+		Column.SetShortName(LOCTEXT("OuterNameColumnName", "Outer"));
+		Column.SetTitleName(LOCTEXT("OuterNameColumnTitle", "Outer Object Name"));
+		Column.SetDescription(LOCTEXT("OuterNameColumnDesc", "The name of the outer UObject of the current UObject instance"));
+
+		Column.SetFlags(ETableColumnFlags::ShouldBeVisible |
+						ETableColumnFlags::CanBeHidden |
+						ETableColumnFlags::CanBeFiltered);
+
+		Column.SetHorizontalAlignment(HAlign_Left);
+		Column.SetInitialWidth(120.0f);
+
+		Column.SetDataType(ETableCellDataType::CString);
+
+		class FValueGetter : public FAggregatedValueGetter
+		{
+		public:
+			virtual const TOptional<FTableCellValue> GetNodeValue(const FBaseTreeNode& Node) const override
+			{
+				if (Node.Is<FObjectNode>())
+				{
+					const FObjectNode& ObjectNode = Node.As<FObjectNode>();
+					TSharedPtr<FObjectNode> OuterNode = ObjectNode.GetOuter();
+					if (OuterNode.IsValid())
+					{
+						return FTableCellValue(OuterNode->GetObjectName());
+					}
+				}
+				return TOptional<FTableCellValue>();
+			}
+		};
+		TSharedRef<ITableCellValueGetter> Getter = MakeShared<FValueGetter>();
+		Column.SetValueGetter(Getter);
+
+		TSharedRef<ITableCellValueFormatter> Formatter = MakeShared<FCStringValueFormatterAsText>();
+		Column.SetValueFormatter(Formatter);
+
+		TSharedRef<ITableCellValueSorter> Sorter = MakeShared<FSorterByCStringValue>(ColumnRef);
+		Column.SetValueSorter(Sorter);
+
+		Column.SetAggregation(ETableColumnAggregation::SameValue);
+
+		AddColumn(ColumnRef);
+	}
+	//////////////////////////////////////////////////
+	// Outer Id Column
+	{
+		TSharedRef<FTableColumn> ColumnRef = MakeShared<FTableColumn>(FObjectTableColumns::OuterIdColumnId);
+		FTableColumn& Column = *ColumnRef;
+
+		Column.SetIndex(ColumnIndex++);
+
+		Column.SetShortName(LOCTEXT("OuterIdColumnName", "Outer Id"));
+		Column.SetTitleName(LOCTEXT("OuterIdColumnTitle", "Outer Id"));
+		Column.SetDescription(LOCTEXT("OuterIdColumnDesc", "The id of the outer UObject of current UObject instance"));
+
+		Column.SetFlags(ETableColumnFlags::CanBeHidden |
+						ETableColumnFlags::CanBeFiltered);
+
+		Column.SetHorizontalAlignment(HAlign_Left);
+		Column.SetInitialWidth(60.0f);
+
+		Column.SetDataType(ETableCellDataType::Int64);
+
+		class FValueGetter : public FAggregatedValueGetter
+		{
+		public:
+			virtual const TOptional<FTableCellValue> GetNodeValue(const FBaseTreeNode& Node) const override
+			{
+				if (Node.Is<FObjectNode>())
+				{
+					const FObjectNode& ObjectNode = Node.As<FObjectNode>();
+					const uint32 Id = ObjectNode.GetOuterId();
+					if (Id != uint32(-1))
+					{
+						return FTableCellValue(int64(Id));
+					}
+				}
+				return TOptional<FTableCellValue>();
+			}
+		};
+		TSharedRef<ITableCellValueGetter> Getter = MakeShared<FValueGetter>();
+		Column.SetValueGetter(Getter);
+
+		TSharedRef<ITableCellValueFormatter> Formatter = MakeShared<FInt64ValueFormatterAsNumber>();
+		Column.SetValueFormatter(Formatter);
+
+		TSharedRef<ITableCellValueSorter> Sorter = MakeShared<FSorterByInt64Value>(ColumnRef);
+		Column.SetValueSorter(Sorter);
+		Column.SetInitialSortMode(EColumnSortMode::Descending);
+
+		//TSharedRef<IFilterValueConverter> Converter = MakeShared<FMemoryFilterValueConverter>();
+		//Column.SetValueConverter(Converter);
+
+		Column.SetAggregation(ETableColumnAggregation::SameValue);
+
+		AddColumn(ColumnRef);
+	}
+	//////////////////////////////////////////////////
+	// Package Name Column
+	{
+		TSharedRef<FTableColumn> ColumnRef = MakeShared<FTableColumn>(FObjectTableColumns::PackageNameColumnId);
+		FTableColumn& Column = *ColumnRef;
+
+		Column.SetIndex(ColumnIndex++);
+
+		Column.SetShortName(LOCTEXT("PackageNameColumnName", "Package"));
+		Column.SetTitleName(LOCTEXT("PackageNameColumnTitle", "Package Name"));
+		Column.SetDescription(LOCTEXT("PackageNameColumnDesc", "The name of the UPackage of the current UObject instance"));
+
+		Column.SetFlags(ETableColumnFlags::ShouldBeVisible |
+						ETableColumnFlags::CanBeHidden |
+						ETableColumnFlags::CanBeFiltered);
+
+		Column.SetHorizontalAlignment(HAlign_Left);
+		Column.SetInitialWidth(120.0f);
+
+		Column.SetDataType(ETableCellDataType::CString);
+
+		class FValueGetter : public FAggregatedValueGetter
+		{
+		public:
+			virtual const TOptional<FTableCellValue> GetNodeValue(const FBaseTreeNode& Node) const override
+			{
+				if (Node.Is<FObjectNode>())
+				{
+					const FObjectNode& ObjectNode = Node.As<FObjectNode>();
+					TSharedPtr<FObjectNode> PackageNode = ObjectNode.GetPackage();
+					if (PackageNode.IsValid())
+					{
+						return FTableCellValue(PackageNode->GetObjectName());
+					}
+				}
+				return TOptional<FTableCellValue>();
+			}
+		};
+		TSharedRef<ITableCellValueGetter> Getter = MakeShared<FValueGetter>();
+		Column.SetValueGetter(Getter);
+
+		TSharedRef<ITableCellValueFormatter> Formatter = MakeShared<FCStringValueFormatterAsText>();
+		Column.SetValueFormatter(Formatter);
+
+		TSharedRef<ITableCellValueSorter> Sorter = MakeShared<FSorterByCStringValue>(ColumnRef);
+		Column.SetValueSorter(Sorter);
+
+		Column.SetAggregation(ETableColumnAggregation::SameValue);
+
+		AddColumn(ColumnRef);
+	}
+	//////////////////////////////////////////////////
+	// Package Id Column
+	{
+		TSharedRef<FTableColumn> ColumnRef = MakeShared<FTableColumn>(FObjectTableColumns::PackageIdColumnId);
+		FTableColumn& Column = *ColumnRef;
+
+		Column.SetIndex(ColumnIndex++);
+
+		Column.SetShortName(LOCTEXT("PackageIdColumnName", "Package Id"));
+		Column.SetTitleName(LOCTEXT("PackageIdColumnTitle", "Package Id"));
+		Column.SetDescription(LOCTEXT("PackageIdColumnDesc", "The id of the UPackage of the current UObject instance"));
+
+		Column.SetFlags(ETableColumnFlags::CanBeHidden |
+						ETableColumnFlags::CanBeFiltered);
+
+		Column.SetHorizontalAlignment(HAlign_Left);
+		Column.SetInitialWidth(60.0f);
+
+		Column.SetDataType(ETableCellDataType::Int64);
+
+		class FValueGetter : public FAggregatedValueGetter
+		{
+		public:
+			virtual const TOptional<FTableCellValue> GetNodeValue(const FBaseTreeNode& Node) const override
+			{
+				if (Node.Is<FObjectNode>())
+				{
+					const FObjectNode& ObjectNode = Node.As<FObjectNode>();
+					const uint32 Id = ObjectNode.GetPackageId();
+					if (Id != uint32(-1))
+					{
+						return FTableCellValue(int64(Id));
+					}
+				}
+				return TOptional<FTableCellValue>();
+			}
+		};
+		TSharedRef<ITableCellValueGetter> Getter = MakeShared<FValueGetter>();
+		Column.SetValueGetter(Getter);
+
+		TSharedRef<ITableCellValueFormatter> Formatter = MakeShared<FInt64ValueFormatterAsNumber>();
+		Column.SetValueFormatter(Formatter);
+
+		TSharedRef<ITableCellValueSorter> Sorter = MakeShared<FSorterByInt64Value>(ColumnRef);
+		Column.SetValueSorter(Sorter);
+		Column.SetInitialSortMode(EColumnSortMode::Descending);
+
+		//TSharedRef<IFilterValueConverter> Converter = MakeShared<FMemoryFilterValueConverter>();
+		//Column.SetValueConverter(Converter);
+
+		Column.SetAggregation(ETableColumnAggregation::SameValue);
+
+		AddColumn(ColumnRef);
+	}
+	//////////////////////////////////////////////////
+	// Object Flags Column
+	{
+		TSharedRef<FTableColumn> ColumnRef = MakeShared<FTableColumn>(FObjectTableColumns::ObjectFlagsColumnId);
+		FTableColumn& Column = *ColumnRef;
+
+		Column.SetIndex(ColumnIndex++);
+
+		Column.SetShortName(LOCTEXT("ObjectFlagsColumnName", "Flags"));
+		Column.SetTitleName(LOCTEXT("ObjectFlagsColumnTitle", "Flags"));
+		Column.SetDescription(LOCTEXT("ObjectFlagsColumnDesc", "The flags of the UObject instance"));
+
+		Column.SetFlags(ETableColumnFlags::ShouldBeVisible |
+						ETableColumnFlags::CanBeHidden |
+						ETableColumnFlags::CanBeFiltered);
+
+		Column.SetHorizontalAlignment(HAlign_Left);
+		Column.SetInitialWidth(140.0f);
+
+		Column.SetDataType(ETableCellDataType::Text);
+
+		class FValueGetter : public FAggregatedValueGetter
+		{
+		public:
+			virtual const TOptional<FTableCellValue> GetNodeValue(const FBaseTreeNode& Node) const override
+			{
+				if (Node.Is<FObjectNode>())
+				{
+					const FObjectNode& ObjectNode = Node.As<FObjectNode>();
+					FString FlagsStr = LexToString(EObjectFlags(ObjectNode.GetObjectFlags()));
+					return FTableCellValue(FText::FromString(FlagsStr));
+				}
+				return TOptional<FTableCellValue>();
+			}
+			virtual uint64 GetValueId(const FTableColumn& Column, const FBaseTreeNode& Node) const override
+			{
+				if (Node.Is<FObjectNode>())
+				{
+					const FObjectNode& ObjectNode = Node.As<FObjectNode>();
+					return uint64(ObjectNode.GetObjectFlags());
+				}
+				return 0;
+			}
+		};
+		TSharedRef<ITableCellValueGetter> Getter = MakeShared<FValueGetter>();
+		Column.SetValueGetter(Getter);
+
+		TSharedRef<ITableCellValueFormatter> Formatter = MakeShared<FTextValueFormatter>();
+		Column.SetValueFormatter(Formatter);
+
+		TSharedRef<ITableCellValueSorter> Sorter = MakeShared<FSorterByTextValueWithId>(ColumnRef);
+		Column.SetValueSorter(Sorter);
+		Column.SetInitialSortMode(EColumnSortMode::Descending);
+
+		//TSharedRef<IFilterValueConverter> Converter = MakeShared<FMemoryFilterValueConverter>();
+		//Column.SetValueConverter(Converter);
+
+		Column.SetAggregation(ETableColumnAggregation::SameValue);
+
+		AddColumn(ColumnRef);
+	}
+	//////////////////////////////////////////////////
+	// Object Path Column
+	{
+		TSharedRef<FTableColumn> ColumnRef = MakeShared<FTableColumn>(FObjectTableColumns::ObjectPathColumnId);
+		FTableColumn& Column = *ColumnRef;
+
+		Column.SetIndex(ColumnIndex++);
+
+		Column.SetShortName(LOCTEXT("ObjectPathColumnName", "Path"));
+		Column.SetTitleName(LOCTEXT("ObjectPathColumnTitle", "Object Path"));
+		Column.SetDescription(LOCTEXT("ObjectPathColumnDesc", "The path of the UObject instance"));
+
+		Column.SetFlags(ETableColumnFlags::CanBeHidden |
+						ETableColumnFlags::CanBeFiltered);
+
+		Column.SetHorizontalAlignment(HAlign_Left);
+		Column.SetInitialWidth(200.0f);
+
+		Column.SetDataType(ETableCellDataType::CString);
+
+		class FValueGetter : public FAggregatedValueGetter
+		{
+		public:
+			virtual const TOptional<FTableCellValue> GetNodeValue(const FBaseTreeNode& Node) const override
+			{
+				if (Node.Is<FObjectNode>())
+				{
+					const FObjectNode& ObjectNode = Node.As<FObjectNode>();
+					return FTableCellValue(ObjectNode.GetObjectPath());
+				}
+				return TOptional<FTableCellValue>();
+			}
+		};
+		TSharedRef<ITableCellValueGetter> Getter = MakeShared<FValueGetter>();
+		Column.SetValueGetter(Getter);
+
+		TSharedRef<ITableCellValueFormatter> Formatter = MakeShared<FCStringValueFormatterAsText>();
+		Column.SetValueFormatter(Formatter);
+
+		TSharedRef<ITableCellValueSorter> Sorter = MakeShared<FSorterByCStringValue>(ColumnRef);
+		Column.SetValueSorter(Sorter);
+
+		Column.SetAggregation(ETableColumnAggregation::SameValue);
+
+		AddColumn(ColumnRef);
+	}
+	//////////////////////////////////////////////////
+	// Verse Path Column
+	{
+		TSharedRef<FTableColumn> ColumnRef = MakeShared<FTableColumn>(FObjectTableColumns::VersePathColumnId);
+		FTableColumn& Column = *ColumnRef;
+
+		Column.SetIndex(ColumnIndex++);
+
+		Column.SetShortName(LOCTEXT("VersePathColumnName", "Verse Path"));
+		Column.SetTitleName(LOCTEXT("VersePathColumnTitle", "Verse Path"));
+		Column.SetDescription(LOCTEXT("VersePathColumnDesc", "The Verse path of the UObject instance"));
+
+		Column.SetFlags(ETableColumnFlags::CanBeHidden |
+						ETableColumnFlags::CanBeFiltered);
+
+		Column.SetHorizontalAlignment(HAlign_Left);
+		Column.SetInitialWidth(200.0f);
+
+		Column.SetDataType(ETableCellDataType::CString);
+
+		class FValueGetter : public FAggregatedValueGetter
+		{
+		public:
+			virtual const TOptional<FTableCellValue> GetNodeValue(const FBaseTreeNode& Node) const override
+			{
+				if (Node.Is<FObjectNode>())
+				{
+					const FObjectNode& ObjectNode = Node.As<FObjectNode>();
+					return FTableCellValue(ObjectNode.GetVersePath());
+				}
+				return TOptional<FTableCellValue>();
+			}
+		};
+		TSharedRef<ITableCellValueGetter> Getter = MakeShared<FValueGetter>();
+		Column.SetValueGetter(Getter);
+
+		TSharedRef<ITableCellValueFormatter> Formatter = MakeShared<FCStringValueFormatterAsText>();
+		Column.SetValueFormatter(Formatter);
+
+		TSharedRef<ITableCellValueSorter> Sorter = MakeShared<FSorterByCStringValue>(ColumnRef);
+		Column.SetValueSorter(Sorter);
+
+		Column.SetAggregation(ETableColumnAggregation::SameValue);
+
+		AddColumn(ColumnRef);
+	}
+	//////////////////////////////////////////////////
+	// Structure Size Column
+	{
+		TSharedRef<FTableColumn> ColumnRef = MakeShared<FTableColumn>(FObjectTableColumns::StructureSizeColumnId);
+		FTableColumn& Column = *ColumnRef;
+
+		Column.SetIndex(ColumnIndex++);
+
+		Column.SetShortName(LOCTEXT("StructureSizeColumnName", "Structure Size"));
+		Column.SetTitleName(LOCTEXT("StructureSizeColumnTitle", "Structure Size"));
+		Column.SetDescription(LOCTEXT("StructureSizeColumnDesc", "The estimated structure size (bytes) of the UStruct instance"));
+
+		Column.SetFlags(ETableColumnFlags::ShouldBeVisible |
+						ETableColumnFlags::CanBeHidden |
+						ETableColumnFlags::CanBeFiltered);
+
+		Column.SetHorizontalAlignment(HAlign_Right);
+		Column.SetInitialWidth(100.0f);
+
+		Column.SetDataType(ETableCellDataType::Int64);
+
+		class FValueGetter : public FAggregatedValueGetter
+		{
+		public:
+			virtual const TOptional<FTableCellValue> GetNodeValue(const FBaseTreeNode& Node) const override
+			{
+				if (Node.Is<FObjectNode>())
+				{
+					const FObjectNode& ObjectNode = Node.As<FObjectNode>();
+					return FTableCellValue(int64(ObjectNode.GetStructureSize()));
+				}
+				return TOptional<FTableCellValue>();
+			}
+		};
+		TSharedRef<ITableCellValueGetter> Getter = MakeShared<FValueGetter>();
+		Column.SetValueGetter(Getter);
+
+		TSharedRef<ITableCellValueFormatter> Formatter = MakeShared<FInt64ValueFormatterAsMemory>();
+		Column.SetValueFormatter(Formatter);
+
+		TSharedRef<ITableCellValueSorter> Sorter = MakeShared<FSorterByInt64Value>(ColumnRef);
+		Column.SetValueSorter(Sorter);
+		Column.SetInitialSortMode(EColumnSortMode::Descending);
+
+		//TSharedRef<IFilterValueConverter> Converter = MakeShared<FMemoryFilterValueConverter>();
+		//Column.SetValueConverter(Converter);
+
+		Column.SetAggregation(ETableColumnAggregation::Sum);
+
+		AddColumn(ColumnRef);
+	}
+	//////////////////////////////////////////////////
+	// System Memory Size Column
+	{
+		TSharedRef<FTableColumn> ColumnRef = MakeShared<FTableColumn>(FObjectTableColumns::SystemMemorySizeColumnId);
+		FTableColumn& Column = *ColumnRef;
+
+		Column.SetIndex(ColumnIndex++);
+
+		Column.SetShortName(LOCTEXT("SystemMemorySizeColumnName", "Mem Size"));
+		Column.SetTitleName(LOCTEXT("SystemMemorySizeColumnTitle", "System Memory Size"));
+		Column.SetDescription(LOCTEXT("SystemMemorySizeColumnDesc", "The estimated system memory size (bytes) of the UObject instance"));
+
+		Column.SetFlags(ETableColumnFlags::ShouldBeVisible |
+						ETableColumnFlags::CanBeHidden |
+						ETableColumnFlags::CanBeFiltered);
+
+		Column.SetHorizontalAlignment(HAlign_Right);
+		Column.SetInitialWidth(100.0f);
+
+		Column.SetDataType(ETableCellDataType::Int64);
+
+		class FValueGetter : public FAggregatedValueGetter
+		{
+		public:
+			virtual const TOptional<FTableCellValue> GetNodeValue(const FBaseTreeNode& Node) const override
+			{
+				if (Node.Is<FObjectNode>())
+				{
+					const FObjectNode& ObjectNode = Node.As<FObjectNode>();
+					return FTableCellValue(int64(ObjectNode.GetSystemMemorySize()));
+				}
+				return TOptional<FTableCellValue>();
+			}
+		};
+		TSharedRef<ITableCellValueGetter> Getter = MakeShared<FValueGetter>();
+		Column.SetValueGetter(Getter);
+
+		TSharedRef<ITableCellValueFormatter> Formatter = MakeShared<FInt64ValueFormatterAsMemory>();
+		Column.SetValueFormatter(Formatter);
+
+		TSharedRef<ITableCellValueSorter> Sorter = MakeShared<FSorterByInt64Value>(ColumnRef);
+		Column.SetValueSorter(Sorter);
+		Column.SetInitialSortMode(EColumnSortMode::Descending);
+
+		//TSharedRef<IFilterValueConverter> Converter = MakeShared<FMemoryFilterValueConverter>();
+		//Column.SetValueConverter(Converter);
+
+		Column.SetAggregation(ETableColumnAggregation::Sum);
+
+		AddColumn(ColumnRef);
+	}
+	//////////////////////////////////////////////////
+	// Video Memory Size Column
+	{
+		TSharedRef<FTableColumn> ColumnRef = MakeShared<FTableColumn>(FObjectTableColumns::VideoMemorySizeColumnId);
+		FTableColumn& Column = *ColumnRef;
+
+		Column.SetIndex(ColumnIndex++);
+
+		Column.SetShortName(LOCTEXT("VideoMemorySizeColumnName", "Video Size"));
+		Column.SetTitleName(LOCTEXT("VideoMemorySizeColumnTitle", "Video Memory Size"));
+		Column.SetDescription(LOCTEXT("VideoMemorySizeColumnDesc", "The estimated video memory size (bytes) of the UObject instance"));
+
+		Column.SetFlags(ETableColumnFlags::ShouldBeVisible |
+						ETableColumnFlags::CanBeHidden |
+						ETableColumnFlags::CanBeFiltered);
+
+		Column.SetHorizontalAlignment(HAlign_Right);
+		Column.SetInitialWidth(100.0f);
+
+		Column.SetDataType(ETableCellDataType::Int64);
+
+		class FValueGetter : public FAggregatedValueGetter
+		{
+		public:
+			virtual const TOptional<FTableCellValue> GetNodeValue(const FBaseTreeNode& Node) const override
+			{
+				if (Node.Is<FObjectNode>())
+				{
+					const FObjectNode& ObjectNode = Node.As<FObjectNode>();
+					return FTableCellValue(int64(ObjectNode.GetVideoMemorySize()));
+				}
+				return TOptional<FTableCellValue>();
+			}
+		};
+		TSharedRef<ITableCellValueGetter> Getter = MakeShared<FValueGetter>();
+		Column.SetValueGetter(Getter);
+
+		TSharedRef<ITableCellValueFormatter> Formatter = MakeShared<FInt64ValueFormatterAsMemory>();
+		Column.SetValueFormatter(Formatter);
+
+		TSharedRef<ITableCellValueSorter> Sorter = MakeShared<FSorterByInt64Value>(ColumnRef);
+		Column.SetValueSorter(Sorter);
+		Column.SetInitialSortMode(EColumnSortMode::Descending);
+
+		//TSharedRef<IFilterValueConverter> Converter = MakeShared<FMemoryFilterValueConverter>();
+		//Column.SetValueConverter(Converter);
+
+		Column.SetAggregation(ETableColumnAggregation::Sum);
+
+		AddColumn(ColumnRef);
+	}
+	//////////////////////////////////////////////////
+	// Estimated Size Column
+	{
+		TSharedRef<FTableColumn> ColumnRef = MakeShared<FTableColumn>(FObjectTableColumns::EstimatedSizeColumnId);
+		FTableColumn& Column = *ColumnRef;
+
+		Column.SetIndex(ColumnIndex++);
+
+		Column.SetShortName(LOCTEXT("EstimatedSizeColumnName", "Estimated Size"));
+		Column.SetTitleName(LOCTEXT("EstimatedSizeColumnTitle", "Estimated Size (MiB)"));
+		Column.SetDescription(LOCTEXT("EstimatedSizeColumnDesc", "The estimated memory size, in MiB"));
+
+		Column.SetFlags(ETableColumnFlags::CanBeHidden |
+						ETableColumnFlags::CanBeFiltered |
+						ETableColumnFlags::IsDynamic);
+
+		Column.SetHorizontalAlignment(HAlign_Right);
+		Column.SetInitialWidth(60.0f);
+
+		Column.SetDataType(ETableCellDataType::Int64);
+
+		class FValueGetter : public FAggregatedValueGetter
+		{
+		public:
+			virtual const TOptional<FTableCellValue> GetNodeValue(const FBaseTreeNode& Node) const override
+			{
+				if (Node.Is<FObjectNode>())
+				{
+					const FObjectNode& ObjectNode = Node.As<FObjectNode>();
+					return FTableCellValue(static_cast<int64>(ObjectNode.GetEstimatedSize()));
+				}
+				return TOptional<FTableCellValue>();
+			}
+		};
+		TSharedRef<ITableCellValueGetter> Getter = MakeShared<FValueGetter>();
+		Column.SetValueGetter(Getter);
+
+		class FEstimatedSizeFormatter : public FTableCellValueFormatter
+		{
+		public:
+			virtual FText FormatValue(const TOptional<FTableCellValue>& InValue) const override
+			{
+				if (InValue.IsSet())
+				{
+					const int64 Size = InValue.GetValue().Int64;
+					if (Size > 0)
+					{
+						int64 SizeMB = Size / (1024 * 1024);
+						if (SizeMB == 0)
+						{
+							return LOCTEXT("LessThan1MB", "< 1 MiB");
+						}
+						return FText::Format(LOCTEXT("EstimatedSizeFmt", "{0} MiB"), FText::AsNumber(SizeMB));
+					}
+				}
+				return FText::GetEmpty();
+			}
+			virtual FText FormatValueForTooltip(const TOptional<FTableCellValue>& InValue) const override
+			{
+				if (InValue.IsSet())
+				{
+					const int64 Size = InValue.GetValue().Int64;
+					if (Size == 0)
+					{
+						return FText::GetEmpty(); // LOCTEXT("ZeroMB", "0");
+					}
+					FNumberFormattingOptions Options;
+					Options.MaximumFractionalDigits = 2;
+					return FText::Format(LOCTEXT("EstimatedSizeTooltipFmt", "{0} ({1} bytes)"),
+						FText::AsMemory(Size, &Options),
+						FText::AsNumber(Size));
+				}
+				return FText::GetEmpty();
+			}
+		};
+
+		TSharedRef<ITableCellValueFormatter> Formatter = MakeShared<FEstimatedSizeFormatter>();
+		Column.SetValueFormatter(Formatter);
+
+		TSharedRef<ITableCellValueSorter> Sorter = MakeShared<FSorterByInt64Value>(ColumnRef);
+		Column.SetValueSorter(Sorter);
+		Column.SetInitialSortMode(EColumnSortMode::Descending);
+
+		Column.SetAggregation(ETableColumnAggregation::Sum);
+
+		AddColumn(ColumnRef);
+	}
+	//////////////////////////////////////////////////
+	// Impact Column
+	{
+		TSharedRef<FTableColumn> ColumnRef = MakeShared<FTableColumn>(FObjectTableColumns::ImpactColumnId);
+		FTableColumn& Column = *ColumnRef;
+
+		Column.SetIndex(ColumnIndex++);
+
+		Column.SetShortName(LOCTEXT("ImpactColumnName", "Impact"));
+		Column.SetTitleName(LOCTEXT("ImpactColumnTitle", "Impact (%)"));
+		Column.SetDescription(LOCTEXT("ImpactColumnDesc", "The impact (percentage) of estimated memory size"));
+
+		Column.SetFlags(ETableColumnFlags::CanBeHidden |
+						ETableColumnFlags::CanBeFiltered |
+						ETableColumnFlags::IsDynamic);
+
+		Column.SetHorizontalAlignment(HAlign_Right);
+		Column.SetInitialWidth(60.0f);
+
+		Column.SetDataType(ETableCellDataType::Double);
+
+		class FValueGetter : public FTableCellValueGetter
+		{
+		public:
+			virtual const TOptional<FTableCellValue> GetValue(const FTableColumn& Column, const FBaseTreeNode& Node) const override
+			{
+				int64 EstimatedSize = 0;
+				if (Node.IsGroup() && Node.Is<FTableTreeNode>())
+				{
+					const FTableTreeNode& TableTreeNode = Node.As<FTableTreeNode>();
+					if (TableTreeNode.HasAggregatedValue(FObjectTableColumns::EstimatedSizeColumnId))
+					{
+						const FTableCellValue& Size = TableTreeNode.GetAggregatedValue(FObjectTableColumns::EstimatedSizeColumnId);
+						if (Size.DataType == ETableCellDataType::Int64 &&
+							Size.Int64 > 0)
+						{
+							EstimatedSize = Size.Int64;
+						}
+					}
+				}
+				else if (Node.Is<FObjectNode>())
+				{
+					const FObjectNode& ObjectNode = Node.As<FObjectNode>();
+					EstimatedSize = ObjectNode.GetEstimatedSize();
+				}
+
+				if (EstimatedSize > 0)
+				{
+					auto ParentNode = Node.GetParent();
+					if (ParentNode.IsValid() && ParentNode->IsGroup() && ParentNode->Is<FTableTreeNode>())
+					{
+						const FTableTreeNode& TableTreeNode = ParentNode->As<FTableTreeNode>();
+						if (TableTreeNode.HasAggregatedValue(FObjectTableColumns::EstimatedSizeColumnId))
+						{
+							const FTableCellValue& Total = TableTreeNode.GetAggregatedValue(FObjectTableColumns::EstimatedSizeColumnId);
+							if (Total.DataType == ETableCellDataType::Int64 &&
+								Total.Int64 > 0)
+							{
+								const double Ratio = double(EstimatedSize) / double(Total.Int64);
+								return FTableCellValue(static_cast<double>(Ratio));
+							}
+						}
+					}
+				}
+
+				return TOptional<FTableCellValue>();
+			}
+		};
+		TSharedRef<ITableCellValueGetter> Getter = MakeShared<FValueGetter>();
+		Column.SetValueGetter(Getter);
+
+		class FImpactFormatter : public FTableCellValueFormatter
+		{
+		public:
+			virtual FText FormatValue(const TOptional<FTableCellValue>& InValue) const override
+			{
+				if (InValue.IsSet())
+				{
+					const double Ratio = InValue.GetValue().Double;
+					int64 Percent = FMath::RoundToInt(Ratio * 100.0);
+					return FText::Format(LOCTEXT("ImpactFmt", "{0}%"), FText::AsNumber(Percent));
+				}
+				return FText::GetEmpty();
+			}
+			virtual FText FormatValueForTooltip(const TOptional<FTableCellValue>& InValue) const override
+			{
+				if (InValue.IsSet())
+				{
+					const double Ratio = InValue.GetValue().Double;
+					FNumberFormattingOptions Options;
+					Options.MaximumFractionalDigits = 2;
+					return FText::Format(LOCTEXT("ImpactTooltipFmt", "{0}"), FText::AsPercent(Ratio, &Options));
+				}
+				return FText::GetEmpty();
+			}
+		};
+		TSharedRef<ITableCellValueFormatter> Formatter = MakeShared<FImpactFormatter>();
+		Column.SetValueFormatter(Formatter);
+
+		TSharedRef<ITableCellValueSorter> Sorter = MakeShared<FSorterByDoubleValue>(ColumnRef);
+		Column.SetValueSorter(Sorter);
+		Column.SetInitialSortMode(EColumnSortMode::Descending);
+
+		Column.SetAggregation(ETableColumnAggregation::None);
+
+		AddColumn(ColumnRef);
+	}
+	//////////////////////////////////////////////////
+	// References Column
+	{
+		TSharedRef<FTableColumn> ColumnRef = MakeShared<FTableColumn>(FObjectTableColumns::ReferencesColumnId);
+		FTableColumn& Column = *ColumnRef;
+
+		Column.SetIndex(ColumnIndex++);
+
+		Column.SetShortName(LOCTEXT("ReferencesColumnName", "References"));
+		Column.SetTitleName(LOCTEXT("ReferencesColumnTitle", "References"));
+		Column.SetDescription(LOCTEXT("ReferencesColumnDesc", "The number of asset references"));
+
+		Column.SetFlags(ETableColumnFlags::CanBeHidden |
+						ETableColumnFlags::CanBeFiltered);
+
+		Column.SetHorizontalAlignment(HAlign_Right);
+		Column.SetInitialWidth(60.0f);
+
+		Column.SetDataType(ETableCellDataType::Int64);
+
+		class FValueGetter : public FAggregatedValueGetter
+		{
+		public:
+			virtual const TOptional<FTableCellValue> GetNodeValue(const FBaseTreeNode& Node) const override
+			{
+				if (Node.Is<FObjectNode>())
+				{
+					const FObjectNode& ObjectNode = Node.As<FObjectNode>();
+					const uint32 NumReferences = ObjectNode.GetNumReferences();
+					return FTableCellValue(static_cast<int64>(NumReferences));
+				}
+				return TOptional<FTableCellValue>();
+			}
+		};
+		TSharedRef<ITableCellValueGetter> Getter = MakeShared<FValueGetter>();
+		Column.SetValueGetter(Getter);
+
+		TSharedRef<ITableCellValueFormatter> Formatter = MakeShared<FInt64ValueFormatterAsNumber>();
+		Column.SetValueFormatter(Formatter);
+
+		TSharedRef<ITableCellValueSorter> Sorter = MakeShared<FSorterByInt64Value>(ColumnRef);
+		Column.SetValueSorter(Sorter);
+		Column.SetInitialSortMode(EColumnSortMode::Descending);
+
+		Column.SetAggregation(ETableColumnAggregation::Sum);
+
+		AddColumn(ColumnRef);
+	}
+	//////////////////////////////////////////////////
+}
+
+////////////////////////////////////////////////////////////////////////////////////////////////////
+
+} // namespace UE::Insights::ObjectProfiler
+
+#undef LOCTEXT_NAMESPACE
diff --git a/Engine/Source/Developer/TraceInsights/Private/Insights/ObjectProfiler/ViewModels/ObjectTable.h b/Engine/Source/Developer/TraceInsights/Private/Insights/ObjectProfiler/ViewModels/ObjectTable.h
new file mode 100644
--- /dev/null
+++ b/Engine/Source/Developer/TraceInsights/Private/Insights/ObjectProfiler/ViewModels/ObjectTable.h
@@ -0,0 +1,55 @@
+// Copyright Epic Games, Inc. All Rights Reserved.
+
+#pragma once
+
+#include "CoreTypes.h"
+#include "UObject/NameTypes.h"
+
+// TraceInsightsCore
+#include "InsightsCore/Table/ViewModels/Table.h"
+
+namespace UE::Insights::ObjectProfiler
+{
+
+////////////////////////////////////////////////////////////////////////////////////////////////////
+
+// Column identifiers
+struct FObjectTableColumns
+{
+	static const FName ObjectCountColumnId;
+	static const FName ObjectNameColumnId;
+	static const FName ObjectIdColumnId;
+	static const FName ClassNameColumnId;
+	static const FName ClassIdColumnId;
+	static const FName OuterNameColumnId;
+	static const FName OuterIdColumnId;
+	static const FName PackageNameColumnId;
+	static const FName PackageIdColumnId;
+	static const FName ObjectFlagsColumnId;
+	static const FName ObjectPathColumnId;
+	static const FName VersePathColumnId;
+	static const FName StructureSizeColumnId;
+	static const FName SystemMemorySizeColumnId;
+	static const FName VideoMemorySizeColumnId;
+	static const FName EstimatedSizeColumnId;
+	static const FName ImpactColumnId;
+	static const FName ReferencesColumnId;
+};
+
+////////////////////////////////////////////////////////////////////////////////////////////////////
+
+class FObjectTable : public FTable
+{
+public:
+	FObjectTable();
+	virtual ~FObjectTable();
+
+	virtual void Reset();
+
+private:
+	void AddDefaultColumns();
+};
+
+////////////////////////////////////////////////////////////////////////////////////////////////////
+
+} // namespace UE::Insights::ObjectProfiler
diff --git a/Engine/Source/Developer/TraceInsights/Private/Insights/ObjectProfiler/ViewModels/ObjectTableViewPresets.cpp b/Engine/Source/Developer/TraceInsights/Private/Insights/ObjectProfiler/ViewModels/ObjectTableViewPresets.cpp
new file mode 100644
--- /dev/null
+++ b/Engine/Source/Developer/TraceInsights/Private/Insights/ObjectProfiler/ViewModels/ObjectTableViewPresets.cpp
@@ -0,0 +1,361 @@
+// Copyright Epic Games, Inc. All Rights Reserved.
+
+#include "ObjectTableViewPresets.h"
+
+// TraceInsightsCore
+#include "InsightsCore/Table/ViewModels/TreeNodeGrouping.h"
+#include "InsightsCore/Table/Widgets/STableTreeView.h"
+
+// TraceInsights
+#include "Insights/ObjectProfiler/ViewModels/ObjectGroupingByOuter.h"
+#include "Insights/ObjectProfiler/ViewModels/ObjectTable.h"
+#include "Insights/ObjectProfiler/Widgets/SObjectTableTreeView.h"
+
+#define LOCTEXT_NAMESPACE "UE::Insights::ObjectProfiler::SObjectTableTreeView"
+
+namespace UE::Insights::ObjectProfiler
+{
+
+TSharedRef<ITableTreeViewPreset> FObjectTableViewPresets::CreateAssetViewPreset(SObjectTableTreeView& TableTreeView)
+{
+	class FAssetViewPreset : public ITableTreeViewPreset
+	{
+	public:
+		virtual FText GetName() const override
+		{
+			return LOCTEXT("Assets_PresetName", "Assets");
+		}
+		virtual FText GetToolTip() const override
+		{
+			return LOCTEXT("Assets_PresetToolTip", "Assets View\nConfigure the tree view to show memory estimations for assets.");
+		}
+		virtual FName GetSortColumn() const override
+		{
+			return FObjectTableColumns::EstimatedSizeColumnId;
+		}
+		virtual EColumnSortMode::Type GetSortMode() const override
+		{
+			return EColumnSortMode::Type::Descending;
+		}
+		virtual void SetCurrentGroupings(const TArray<TSharedPtr<FTreeNodeGrouping>>& InAvailableGroupings, TArray<TSharedPtr<FTreeNodeGrouping>>& InOutCurrentGroupings) const override
+		{
+			InOutCurrentGroupings.Reset();
+
+			check(InAvailableGroupings.Num() > 0);
+			check(InAvailableGroupings[0]->Is<FTreeNodeGroupingFlat>());
+			InOutCurrentGroupings.Add(InAvailableGroupings[0]);
+
+			auto GroupingByClass = InAvailableGroupings.FindByPredicate(
+				[](TSharedPtr<FTreeNodeGrouping>& Grouping)
+				{
+					return Grouping->Is<FTreeNodeGroupingByUniqueValue>() &&
+						Grouping->As<FTreeNodeGroupingByUniqueValue>().GetColumnId() == FObjectTableColumns::ClassNameColumnId;
+				});
+			if (GroupingByClass)
+			{
+				InOutCurrentGroupings.Add(*GroupingByClass);
+			}
+		}
+		virtual void GetColumnConfigSet(TArray<FTableColumnConfig>& InOutConfigSet) const override
+		{
+			InOutConfigSet.Add({ FTable::GetHierarchyColumnId(),                true, 300.0f });
+			InOutConfigSet.Add({ FObjectTableColumns::ObjectCountColumnId,      !true, 60.0f });
+			InOutConfigSet.Add({ FObjectTableColumns::ObjectNameColumnId,       !true, 120.0f });
+			InOutConfigSet.Add({ FObjectTableColumns::ObjectIdColumnId,         !true, 60.0f });
+			InOutConfigSet.Add({ FObjectTableColumns::ClassNameColumnId,        !true, 120.0f });
+			InOutConfigSet.Add({ FObjectTableColumns::ClassIdColumnId,          !true, 60.0f });
+			InOutConfigSet.Add({ FObjectTableColumns::OuterNameColumnId,        !true, 120.0f });
+			InOutConfigSet.Add({ FObjectTableColumns::OuterIdColumnId,          !true, 60.0f });
+			InOutConfigSet.Add({ FObjectTableColumns::PackageNameColumnId,      !true, 120.0f });
+			InOutConfigSet.Add({ FObjectTableColumns::PackageIdColumnId,        !true, 60.0f });
+			InOutConfigSet.Add({ FObjectTableColumns::ObjectFlagsColumnId,      !true, 140.0f });
+			InOutConfigSet.Add({ FObjectTableColumns::ObjectPathColumnId,       !true, 200.0f });
+			InOutConfigSet.Add({ FObjectTableColumns::VersePathColumnId,        !true, 200.0f });
+			InOutConfigSet.Add({ FObjectTableColumns::StructureSizeColumnId,    !true, 100.0f });
+			InOutConfigSet.Add({ FObjectTableColumns::SystemMemorySizeColumnId, !true, 100.0f });
+			InOutConfigSet.Add({ FObjectTableColumns::VideoMemorySizeColumnId,  !true, 100.0f });
+			InOutConfigSet.Add({ FObjectTableColumns::EstimatedSizeColumnId,    true, 80.0f });
+			InOutConfigSet.Add({ FObjectTableColumns::ImpactColumnId,           true, 80.0f });
+			InOutConfigSet.Add({ FObjectTableColumns::ReferencesColumnId,       true, 80.0f });
+		}
+		virtual void OnAppliedToView(STableTreeView& TableTreeView) const override
+		{
+			//SObjectTableTreeView& ObjectTreeView = (SObjectTableTreeView&)TableTreeView;
+		}
+	};
+	return MakeShared<FAssetViewPreset>();
+}
+
+////////////////////////////////////////////////////////////////////////////////////////////////////
+// Objects View
+
+TSharedRef<ITableTreeViewPreset> FObjectTableViewPresets::CreateObjectViewPreset(SObjectTableTreeView& TableTreeView)
+{
+	class FObjectViewPreset : public ITableTreeViewPreset
+	{
+	public:
+		virtual FText GetName() const override
+		{
+			return LOCTEXT("Objects_PresetName", "Objects");
+		}
+		virtual FText GetToolTip() const override
+		{
+			return LOCTEXT("Objects_PresetToolTip", "Objects View\nConfigure the tree view to show the all UObject instances.");
+		}
+		virtual FName GetSortColumn() const override
+		{
+			return FObjectTableColumns::ObjectIdColumnId;
+		}
+		virtual EColumnSortMode::Type GetSortMode() const override
+		{
+			return EColumnSortMode::Type::Ascending;
+		}
+		virtual void SetCurrentGroupings(const TArray<TSharedPtr<FTreeNodeGrouping>>& InAvailableGroupings, TArray<TSharedPtr<FTreeNodeGrouping>>& InOutCurrentGroupings) const override
+		{
+			InOutCurrentGroupings.Reset();
+
+			check(InAvailableGroupings.Num() > 0);
+			check(InAvailableGroupings[0]->Is<FTreeNodeGroupingFlat>());
+			InOutCurrentGroupings.Add(InAvailableGroupings[0]);
+		}
+		virtual void GetColumnConfigSet(TArray<FTableColumnConfig>& InOutConfigSet) const override
+		{
+			InOutConfigSet.Add({ FTable::GetHierarchyColumnId(),                true, 400.0f });
+			InOutConfigSet.Add({ FObjectTableColumns::ObjectCountColumnId,      true, 60.0f });
+			InOutConfigSet.Add({ FObjectTableColumns::ObjectNameColumnId,       !true, 120.0f });
+			InOutConfigSet.Add({ FObjectTableColumns::ObjectIdColumnId,         !true, 60.0f });
+			InOutConfigSet.Add({ FObjectTableColumns::ClassNameColumnId,        true, 120.0f });
+			InOutConfigSet.Add({ FObjectTableColumns::ClassIdColumnId,          !true, 60.0f });
+			InOutConfigSet.Add({ FObjectTableColumns::OuterNameColumnId,        true, 120.0f });
+			InOutConfigSet.Add({ FObjectTableColumns::OuterIdColumnId,          !true, 60.0f });
+			InOutConfigSet.Add({ FObjectTableColumns::PackageNameColumnId,      true, 120.0f });
+			InOutConfigSet.Add({ FObjectTableColumns::PackageIdColumnId,        !true, 60.0f });
+			InOutConfigSet.Add({ FObjectTableColumns::ObjectFlagsColumnId,      true, 140.0f });
+			InOutConfigSet.Add({ FObjectTableColumns::ObjectPathColumnId,       !true, 200.0f });
+			InOutConfigSet.Add({ FObjectTableColumns::VersePathColumnId,        !true, 200.0f });
+			InOutConfigSet.Add({ FObjectTableColumns::StructureSizeColumnId,    true, 100.0f });
+			InOutConfigSet.Add({ FObjectTableColumns::SystemMemorySizeColumnId, true, 100.0f });
+			InOutConfigSet.Add({ FObjectTableColumns::VideoMemorySizeColumnId,  true, 100.0f });
+		}
+		virtual void OnAppliedToView(STableTreeView& TableTreeView) const override
+		{
+			//SObjectTableTreeView& ObjectTreeView = (SObjectTableTreeView&)TableTreeView;
+		}
+	};
+	return MakeShared<FObjectViewPreset>();
+}
+
+////////////////////////////////////////////////////////////////////////////////////////////////////
+// Class View
+
+TSharedRef<ITableTreeViewPreset> FObjectTableViewPresets::CreateClassViewPreset(SObjectTableTreeView& TableTreeView)
+{
+	class FClassViewPreset : public ITableTreeViewPreset
+	{
+	public:
+		virtual FText GetName() const override
+		{
+			return LOCTEXT("Class_PresetName", "Classes");
+		}
+		virtual FText GetToolTip() const override
+		{
+			return LOCTEXT("Class_PresetToolTip", "Classes View\nConfigure the tree view to show the all UClass instances.");
+		}
+		virtual FName GetSortColumn() const override
+		{
+			return FObjectTableColumns::ObjectCountColumnId;
+		}
+		virtual EColumnSortMode::Type GetSortMode() const override
+		{
+			return EColumnSortMode::Type::Descending;
+		}
+		virtual void SetCurrentGroupings(const TArray<TSharedPtr<FTreeNodeGrouping>>& InAvailableGroupings, TArray<TSharedPtr<FTreeNodeGrouping>>& InOutCurrentGroupings) const override
+		{
+			InOutCurrentGroupings.Reset();
+
+			check(InAvailableGroupings.Num() > 0);
+			check(InAvailableGroupings[0]->Is<FTreeNodeGroupingFlat>());
+			InOutCurrentGroupings.Add(InAvailableGroupings[0]);
+
+			auto GroupingByClass = InAvailableGroupings.FindByPredicate(
+				[](TSharedPtr<FTreeNodeGrouping>& Grouping)
+				{
+					return Grouping->Is<FTreeNodeGroupingByUniqueValue>() &&
+						Grouping->As<FTreeNodeGroupingByUniqueValue>().GetColumnId() == FObjectTableColumns::ClassNameColumnId;
+				});
+			if (GroupingByClass)
+			{
+				InOutCurrentGroupings.Add(*GroupingByClass);
+			}
+		}
+		virtual void GetColumnConfigSet(TArray<FTableColumnConfig>& InOutConfigSet) const override
+		{
+			InOutConfigSet.Add({ FTable::GetHierarchyColumnId(),                true, 400.0f });
+			InOutConfigSet.Add({ FObjectTableColumns::ObjectCountColumnId,      true, 60.0f });
+			InOutConfigSet.Add({ FObjectTableColumns::ObjectNameColumnId,       !true, 120.0f });
+			InOutConfigSet.Add({ FObjectTableColumns::ObjectIdColumnId,         !true, 60.0f });
+			InOutConfigSet.Add({ FObjectTableColumns::ClassNameColumnId,        !true, 120.0f });
+			InOutConfigSet.Add({ FObjectTableColumns::ClassIdColumnId,          !true, 60.0f });
+			InOutConfigSet.Add({ FObjectTableColumns::OuterNameColumnId,        true, 120.0f });
+			InOutConfigSet.Add({ FObjectTableColumns::OuterIdColumnId,          !true, 60.0f });
+			InOutConfigSet.Add({ FObjectTableColumns::PackageNameColumnId,      true, 120.0f });
+			InOutConfigSet.Add({ FObjectTableColumns::PackageIdColumnId,        !true, 60.0f });
+			InOutConfigSet.Add({ FObjectTableColumns::ObjectFlagsColumnId,      true, 140.0f });
+			InOutConfigSet.Add({ FObjectTableColumns::ObjectPathColumnId,       !true, 200.0f });
+			InOutConfigSet.Add({ FObjectTableColumns::VersePathColumnId,        !true, 200.0f });
+			InOutConfigSet.Add({ FObjectTableColumns::StructureSizeColumnId,    true, 100.0f });
+			InOutConfigSet.Add({ FObjectTableColumns::SystemMemorySizeColumnId, true, 100.0f });
+			InOutConfigSet.Add({ FObjectTableColumns::VideoMemorySizeColumnId,  true, 100.0f });
+		}
+		virtual void OnAppliedToView(STableTreeView& TableTreeView) const override
+		{
+			//SObjectTableTreeView& ObjectTreeView = (SObjectTableTreeView&)TableTreeView;
+		}
+	};
+	return MakeShared<FClassViewPreset>();
+}
+
+////////////////////////////////////////////////////////////////////////////////////////////////////
+// Outer View
+
+TSharedRef<ITableTreeViewPreset> FObjectTableViewPresets::CreateOuterViewPreset(SObjectTableTreeView& TableTreeView)
+{
+	class FOuterViewPreset : public ITableTreeViewPreset
+	{
+	public:
+		virtual FText GetName() const override
+		{
+			return LOCTEXT("Outer_PresetName", "Object Hierarchy");
+		}
+		virtual FText GetToolTip() const override
+		{
+			return LOCTEXT("Outer_PresetToolTip", "Object Hierarchy View\nConfigure the tree view to show the all UObject instances in a hierarchical view.");
+		}
+		virtual FName GetSortColumn() const override
+		{
+			return FObjectTableColumns::ObjectNameColumnId;
+		}
+		virtual EColumnSortMode::Type GetSortMode() const override
+		{
+			return EColumnSortMode::Type::Ascending;
+		}
+		virtual void SetCurrentGroupings(const TArray<TSharedPtr<FTreeNodeGrouping>>& InAvailableGroupings, TArray<TSharedPtr<FTreeNodeGrouping>>& InOutCurrentGroupings) const override
+		{
+			InOutCurrentGroupings.Reset();
+
+			check(InAvailableGroupings.Num() > 0);
+			check(InAvailableGroupings[0]->Is<FTreeNodeGroupingFlat>());
+			InOutCurrentGroupings.Add(InAvailableGroupings[0]);
+
+			auto GroupingByOuter = InAvailableGroupings.FindByPredicate(
+				[](TSharedPtr<FTreeNodeGrouping>& Grouping)
+				{
+					return Grouping->Is<FObjectGroupingByOuter>();
+				});
+			if (GroupingByOuter)
+			{
+				InOutCurrentGroupings.Add(*GroupingByOuter);
+			}
+		}
+		virtual void GetColumnConfigSet(TArray<FTableColumnConfig>& InOutConfigSet) const override
+		{
+			InOutConfigSet.Add({ FTable::GetHierarchyColumnId(),                true, 400.0f });
+			InOutConfigSet.Add({ FObjectTableColumns::ObjectCountColumnId,      true, 60.0f });
+			InOutConfigSet.Add({ FObjectTableColumns::ObjectNameColumnId,       !true, 120.0f });
+			InOutConfigSet.Add({ FObjectTableColumns::ObjectIdColumnId,         !true, 60.0f });
+			InOutConfigSet.Add({ FObjectTableColumns::ClassNameColumnId,        true, 120.0f });
+			InOutConfigSet.Add({ FObjectTableColumns::ClassIdColumnId,          !true, 60.0f });
+			InOutConfigSet.Add({ FObjectTableColumns::OuterNameColumnId,        !true, 120.0f });
+			InOutConfigSet.Add({ FObjectTableColumns::OuterIdColumnId,          !true, 60.0f });
+			InOutConfigSet.Add({ FObjectTableColumns::PackageNameColumnId,      true, 120.0f });
+			InOutConfigSet.Add({ FObjectTableColumns::PackageIdColumnId,        !true, 60.0f });
+			InOutConfigSet.Add({ FObjectTableColumns::ObjectFlagsColumnId,      true, 140.0f });
+			InOutConfigSet.Add({ FObjectTableColumns::ObjectPathColumnId,       !true, 200.0f });
+			InOutConfigSet.Add({ FObjectTableColumns::VersePathColumnId,        !true, 200.0f });
+			InOutConfigSet.Add({ FObjectTableColumns::StructureSizeColumnId,    true, 100.0f });
+			InOutConfigSet.Add({ FObjectTableColumns::SystemMemorySizeColumnId, true, 100.0f });
+			InOutConfigSet.Add({ FObjectTableColumns::VideoMemorySizeColumnId,  true, 100.0f });
+		}
+		virtual void OnAppliedToView(STableTreeView& TableTreeView) const override
+		{
+			//SObjectTableTreeView& ObjectTreeView = (SObjectTableTreeView&)TableTreeView;
+		}
+	};
+	return MakeShared<FOuterViewPreset>();
+}
+
+////////////////////////////////////////////////////////////////////////////////////////////////////
+// Package View
+
+TSharedRef<ITableTreeViewPreset> FObjectTableViewPresets::CreatePackageViewPreset(SObjectTableTreeView& TableTreeView)
+{
+	class FPackageViewPreset : public ITableTreeViewPreset
+	{
+	public:
+		virtual FText GetName() const override
+		{
+			return LOCTEXT("Package_PresetName", "Packages");
+		}
+		virtual FText GetToolTip() const override
+		{
+			return LOCTEXT("Package_PresetToolTip", "Packages View\nConfigure the tree view to show the all UPackage instances.");
+		}
+		virtual FName GetSortColumn() const override
+		{
+			return FObjectTableColumns::ObjectNameColumnId;
+		}
+		virtual EColumnSortMode::Type GetSortMode() const override
+		{
+			return EColumnSortMode::Type::Ascending;
+		}
+		virtual void SetCurrentGroupings(const TArray<TSharedPtr<FTreeNodeGrouping>>& InAvailableGroupings, TArray<TSharedPtr<FTreeNodeGrouping>>& InOutCurrentGroupings) const override
+		{
+			InOutCurrentGroupings.Reset();
+
+			check(InAvailableGroupings.Num() > 0);
+			check(InAvailableGroupings[0]->Is<FTreeNodeGroupingFlat>());
+			InOutCurrentGroupings.Add(InAvailableGroupings[0]);
+
+			auto GroupingByPackage = InAvailableGroupings.FindByPredicate(
+				[](TSharedPtr<FTreeNodeGrouping>& Grouping)
+				{
+					return Grouping->Is<FTreeNodeGroupingByUniqueValue>() &&
+						Grouping->As<FTreeNodeGroupingByUniqueValue>().GetColumnId() == FObjectTableColumns::PackageNameColumnId;
+				});
+			if (GroupingByPackage)
+			{
+				InOutCurrentGroupings.Add(*GroupingByPackage);
+			}
+		}
+		virtual void GetColumnConfigSet(TArray<FTableColumnConfig>& InOutConfigSet) const override
+		{
+			InOutConfigSet.Add({ FTable::GetHierarchyColumnId(),                true, 400.0f });
+			InOutConfigSet.Add({ FObjectTableColumns::ObjectCountColumnId,      true, 60.0f });
+			InOutConfigSet.Add({ FObjectTableColumns::ObjectNameColumnId,       !true, 120.0f });
+			InOutConfigSet.Add({ FObjectTableColumns::ObjectIdColumnId,         !true, 60.0f });
+			InOutConfigSet.Add({ FObjectTableColumns::ClassNameColumnId,        true, 120.0f });
+			InOutConfigSet.Add({ FObjectTableColumns::ClassIdColumnId,          !true, 60.0f });
+			InOutConfigSet.Add({ FObjectTableColumns::OuterNameColumnId,        true, 120.0f });
+			InOutConfigSet.Add({ FObjectTableColumns::OuterIdColumnId,          !true, 60.0f });
+			InOutConfigSet.Add({ FObjectTableColumns::PackageNameColumnId,      !true, 120.0f });
+			InOutConfigSet.Add({ FObjectTableColumns::PackageIdColumnId,        !true, 60.0f });
+			InOutConfigSet.Add({ FObjectTableColumns::ObjectFlagsColumnId,      true, 140.0f });
+			InOutConfigSet.Add({ FObjectTableColumns::ObjectPathColumnId,       !true, 200.0f });
+			InOutConfigSet.Add({ FObjectTableColumns::VersePathColumnId,        !true, 200.0f });
+			InOutConfigSet.Add({ FObjectTableColumns::StructureSizeColumnId,    true, 100.0f });
+			InOutConfigSet.Add({ FObjectTableColumns::SystemMemorySizeColumnId, true, 100.0f });
+			InOutConfigSet.Add({ FObjectTableColumns::VideoMemorySizeColumnId,  true, 100.0f });
+		}
+		virtual void OnAppliedToView(STableTreeView& TableTreeView) const override
+		{
+			//SObjectTableTreeView& ObjectTreeView = (SObjectTableTreeView&)TableTreeView;
+		}
+	};
+	return MakeShared<FPackageViewPreset>();
+}
+
+////////////////////////////////////////////////////////////////////////////////////////////////////
+
+} // namespace UE::Insights::ObjectProfiler
+
+#undef LOCTEXT_NAMESPACE
diff --git a/Engine/Source/Developer/TraceInsights/Private/Insights/ObjectProfiler/ViewModels/ObjectTableViewPresets.h b/Engine/Source/Developer/TraceInsights/Private/Insights/ObjectProfiler/ViewModels/ObjectTableViewPresets.h
new file mode 100644
--- /dev/null
+++ b/Engine/Source/Developer/TraceInsights/Private/Insights/ObjectProfiler/ViewModels/ObjectTableViewPresets.h
@@ -0,0 +1,26 @@
+// Copyright Epic Games, Inc. All Rights Reserved.
+
+#pragma once
+
+#include "CoreTypes.h"
+
+#include "Templates/SharedPointer.h"
+
+namespace UE::Insights { class ITableTreeViewPreset; }
+
+namespace UE::Insights::ObjectProfiler
+{
+
+class SObjectTableTreeView;
+
+class FObjectTableViewPresets
+{
+public:
+	static TSharedRef<ITableTreeViewPreset> CreateAssetViewPreset(SObjectTableTreeView& TableTreeView);
+	static TSharedRef<ITableTreeViewPreset> CreateObjectViewPreset(SObjectTableTreeView& TableTreeView);
+	static TSharedRef<ITableTreeViewPreset> CreateClassViewPreset(SObjectTableTreeView& TableTreeView);
+	static TSharedRef<ITableTreeViewPreset> CreateOuterViewPreset(SObjectTableTreeView& TableTreeView);
+	static TSharedRef<ITableTreeViewPreset> CreatePackageViewPreset(SObjectTableTreeView& TableTreeView);
+};
+
+} // namespace UE::Insights::ObjectProfiler
diff --git a/Engine/Source/Developer/TraceInsights/Private/Insights/ObjectProfiler/Widgets/SObjectTableTreeView.cpp b/Engine/Source/Developer/TraceInsights/Private/Insights/ObjectProfiler/Widgets/SObjectTableTreeView.cpp
new file mode 100644
--- /dev/null
+++ b/Engine/Source/Developer/TraceInsights/Private/Insights/ObjectProfiler/Widgets/SObjectTableTreeView.cpp
@@ -0,0 +1,604 @@
+// Copyright Epic Games, Inc. All Rights Reserved.
+
+#include "SObjectTableTreeView.h"
+
+#include "DesktopPlatformModule.h"
+#include "Framework/MultiBox/MultiBoxBuilder.h"
+#include "Misc/ConfigCacheIni.h"
+#include "SlateOptMacros.h"
+#include "Styling/StarshipCoreStyle.h"
+#include "Styling/StyleColors.h"
+#include "Widgets/Colors/SColorPicker.h"
+#include "Widgets/Images/SImage.h"
+#include "Widgets/Input/SButton.h"
+#include "Widgets/Input/SComboBox.h"
+#include "Widgets/Input/SSegmentedControl.h"
+
+// TraceServices
+#include "Common/ProviderLock.h"
+#include "TraceServices/Model/ObjectProvider.h"
+
+// TraceInsightsCore
+#include "InsightsCore/Common/Stopwatch.h"
+#include "InsightsCore/Table/ViewModels/Table.h"
+
+// TraceInsights
+#include "Insights/InsightsManager.h"
+#include "Insights/InsightsStyle.h"
+#include "Insights/ObjectProfiler/ObjectProfilerManager.h"
+#include "Insights/ObjectProfiler/ViewModels/ObjectGroupingByOuter.h"
+#include "Insights/ObjectProfiler/ViewModels/ObjectNode.h"
+#include "Insights/ObjectProfiler/ViewModels/ObjectTable.h"
+#include "Insights/ObjectProfiler/ViewModels/ObjectTableViewPresets.h"
+
+#define LOCTEXT_NAMESPACE "UE::Insights::ObjectProfiler::SObjectTableTreeView"
+
+namespace UE::Insights::ObjectProfiler
+{
+
+////////////////////////////////////////////////////////////////////////////////////////////////////
+
+SObjectTableTreeView::SObjectTableTreeView()
+{
+	bRunInAsyncMode = true;
+}
+
+////////////////////////////////////////////////////////////////////////////////////////////////////
+
+SObjectTableTreeView::~SObjectTableTreeView()
+{
+}
+
+////////////////////////////////////////////////////////////////////////////////////////////////////
+
+void SObjectTableTreeView::Construct(const FArguments& InArgs)
+{
+	TSharedRef<FObjectTable> TablePtr = MakeShared<FObjectTable>();
+	TablePtr->Reset();
+	TablePtr->SetDisplayName(LOCTEXT("ObjectsTableName", "Objects"));
+
+	ConstructWidget(TablePtr);
+
+	// Apply the default preset.
+	check(GetAvailableViewPresets() != nullptr);
+	check(GetAvailableViewPresets()->Num() > 0);
+	ApplyViewPreset(*(*GetAvailableViewPresets())[0]);
+
+	UpdateSelectionStatsText();
+}
+
+////////////////////////////////////////////////////////////////////////////////////////////////////
+
+void SObjectTableTreeView::ConstructHeaderArea(TSharedRef<SVerticalBox> InHostBox)
+{
+	TSharedPtr<SWidget> FilterToolbarWidget = ConstructFilterToolbar();
+	if (FilterToolbarWidget.IsValid())
+	{
+		InHostBox->AddSlot()
+			.AutoHeight()
+			.Padding(4.0f, 2.0f, 4.0f, 2.0f)
+			[
+				FilterToolbarWidget.ToSharedRef()
+			];
+	}
+
+	TSharedPtr<SWidget> HierarchyBreadcrumbTrailWidget = ConstructHierarchyBreadcrumbTrail();
+	if (HierarchyBreadcrumbTrailWidget.IsValid())
+	{
+		InHostBox->AddSlot()
+			.AutoHeight()
+			.Padding(4.0f, 2.0f, 4.0f, -2.0f)
+			[
+				HierarchyBreadcrumbTrailWidget.ToSharedRef()
+			];
+	}
+}
+
+////////////////////////////////////////////////////////////////////////////////////////////////////
+
+void SObjectTableTreeView::ConstructFooterArea(TSharedRef<SVerticalBox> InHostBox)
+{
+	TSharedPtr<SWidget> FooterWidget = ConstructFooter();
+	if (FooterWidget.IsValid())
+	{
+		InHostBox->AddSlot()
+			.AutoHeight()
+			.Padding(0.0f, 6.0f, 0.0f, 0.0f)
+			[
+				FooterWidget.ToSharedRef()
+			];
+	}
+}
+
+////////////////////////////////////////////////////////////////////////////////////////////////////
+
+BEGIN_SLATE_FUNCTION_BUILD_OPTIMIZATION
+TSharedPtr<SWidget> SObjectTableTreeView::ConstructFilterToolbar()
+{
+	TSharedRef<SHorizontalBox> FilterLine = SNew(SHorizontalBox)
+
+	+ SHorizontalBox::Slot()
+	.FillWidth(1.0f)
+	.VAlign(VAlign_Center)
+	[
+		ConstructSearchBox()
+	]
+
+	+ SHorizontalBox::Slot()
+	.AutoWidth()
+	.Padding(2.0f, 0.0f, 0.0f, 0.0f)
+	.VAlign(VAlign_Center)
+	[
+		ConstructFilterConfiguratorButton()
+	];
+
+	ConstructViewPreset(FilterLine, 140.0f);
+
+	return FilterLine;
+}
+END_SLATE_FUNCTION_BUILD_OPTIMIZATION
+
+////////////////////////////////////////////////////////////////////////////////////////////////////
+
+BEGIN_SLATE_FUNCTION_BUILD_OPTIMIZATION
+TSharedPtr<SWidget> SObjectTableTreeView::ConstructToolbar()
+{
+	return nullptr;
+}
+END_SLATE_FUNCTION_BUILD_OPTIMIZATION
+
+////////////////////////////////////////////////////////////////////////////////////////////////////
+
+BEGIN_SLATE_FUNCTION_BUILD_OPTIMIZATION
+TSharedPtr<SWidget> SObjectTableTreeView::ConstructFooter()
+{
+	return
+		SNew(SHorizontalBox)
+
+		+ SHorizontalBox::Slot()
+		.FillWidth(1.0f)
+
+		+ SHorizontalBox::Slot()
+		.AutoWidth()
+		.Padding(0.0f, 2.0f, 0.0f, 2.0f)
+		[
+			SNew(STextBlock)
+			.Text(this, &SObjectTableTreeView::GetNumSelectedObjectsText)
+		]
+
+		+ SHorizontalBox::Slot()
+		.AutoWidth()
+		.Padding(4.0f, 2.0f, 0.0f, 2.0f)
+		[
+			SNew(STextBlock)
+			.Text(this, &SObjectTableTreeView::GetSelectedObjectsText)
+			.ColorAndOpacity(FSlateColor(EStyleColor::White25))
+		]
+
+		+ SHorizontalBox::Slot()
+		.AutoWidth()
+		.Padding(0.0f, 2.0f, 0.0f, 2.0f)
+		[
+			SNew(STextBlock)
+			.Text(FText::FromString(TEXT(" (StructSize: ")))
+			.ColorAndOpacity(FSlateColor(EStyleColor::White25))
+			.Visibility_Lambda([this]() { return !NumSelectedObjectsText.IsEmpty() ? EVisibility::Visible : EVisibility::Collapsed; })
+		]
+
+		+ SHorizontalBox::Slot()
+		.AutoWidth()
+		.Padding(0.0f, 2.0f, 0.0f, 2.0f)
+		[
+			SNew(STextBlock)
+			.Text(this, &SObjectTableTreeView::GetSelectionStructSizeText)
+		]
+
+		+ SHorizontalBox::Slot()
+		.AutoWidth()
+		.Padding(0.0f, 2.0f, 0.0f, 2.0f)
+		[
+			SNew(STextBlock)
+			.Text(FText::FromString(TEXT(", SystemMemSize: ")))
+			.ColorAndOpacity(FSlateColor(EStyleColor::White25))
+			.Visibility_Lambda([this]() { return !NumSelectedObjectsText.IsEmpty() ? EVisibility::Visible : EVisibility::Collapsed; })
+		]
+
+		+ SHorizontalBox::Slot()
+		.AutoWidth()
+		.Padding(0.0f, 2.0f, 0.0f, 2.0f)
+		[
+			SNew(STextBlock)
+			.Text(this, &SObjectTableTreeView::GetSelectionSystemMemSizeText)
+		]
+
+		+ SHorizontalBox::Slot()
+		.AutoWidth()
+		.Padding(0.0f, 2.0f, 0.0f, 2.0f)
+		[
+			SNew(STextBlock)
+			.Text(FText::FromString(TEXT(", VideoMemSize: ")))
+			.ColorAndOpacity(FSlateColor(EStyleColor::White25))
+			.Visibility_Lambda([this]() { return !NumSelectedObjectsText.IsEmpty() ? EVisibility::Visible : EVisibility::Collapsed; })
+		]
+
+		+ SHorizontalBox::Slot()
+		.AutoWidth()
+		.Padding(0.0f, 2.0f, 0.0f, 2.0f)
+		[
+			SNew(STextBlock)
+			.Text(this, &SObjectTableTreeView::GetSelectionVideoMemSizeText)
+		]
+
+		+ SHorizontalBox::Slot()
+		.AutoWidth()
+		.Padding(0.0f, 2.0f, 0.0f, 2.0f)
+		[
+			SNew(STextBlock)
+			.Text(FText::FromString(TEXT(")")))
+			.ColorAndOpacity(FSlateColor(EStyleColor::White25))
+			.Visibility_Lambda([this]() { return !NumSelectedObjectsText.IsEmpty() ? EVisibility::Visible : EVisibility::Collapsed; })
+		]
+
+		+ SHorizontalBox::Slot()
+		.FillWidth(1.0f);
+}
+END_SLATE_FUNCTION_BUILD_OPTIMIZATION
+
+////////////////////////////////////////////////////////////////////////////////////////////////////
+
+void SObjectTableTreeView::UpdateSelectionStatsText()
+{
+	TArray<FTableTreeNodePtr> SelectedNodes;
+	const int32 NumSelectedNodes = TreeView->GetSelectedItems(SelectedNodes);
+
+	if (NumSelectedNodes > 0)
+	{
+		int64 TotalCount = 0;
+		int64 TotalStructSize = 0;
+		int64 TotalSystemMemSize = 0;
+		int64 TotalVideoMemSize = 0;
+		for (const FTableTreeNodePtr& SelectedNode : SelectedNodes)
+		{
+			if (SelectedNode->Is<FObjectNode>())
+			{
+				const FObjectNode& ObjectNode = SelectedNode->As<FObjectNode>();
+				++TotalCount;
+				TotalStructSize += ObjectNode.GetStructureSize();
+				TotalSystemMemSize += ObjectNode.GetSystemMemorySize();
+				TotalVideoMemSize += ObjectNode.GetVideoMemorySize();
+			}
+		}
+
+		FNumberFormattingOptions FormattingOptionsMem;
+		FormattingOptionsMem.MaximumFractionalDigits = 2;
+
+		NumSelectedObjectsText = FText::AsNumber(TotalCount);
+		SelectedObjectsText = FText::Format(LOCTEXT("SelectionStatsFmt", "selected {0}|plural(one=object,other=objects)"), FText::AsNumber(TotalCount));
+		SelectionStructSizeText = TotalStructSize == 0 ? FText::FromString(TEXT("0")) : FText::AsMemory(TotalStructSize, &FormattingOptionsMem);
+		SelectionSystemMemSizeText = TotalSystemMemSize == 0 ? FText::FromString(TEXT("0")) : FText::AsMemory(TotalSystemMemSize, &FormattingOptionsMem);
+		SelectionVideoMemSizeText = TotalVideoMemSize == 0 ? FText::FromString(TEXT("0")) : FText::AsMemory(TotalVideoMemSize, &FormattingOptionsMem);
+	}
+	else
+	{
+		NumSelectedObjectsText = FText::GetEmpty();
+		SelectedObjectsText = LOCTEXT("NoSelectionStats", "No object selected");
+		SelectionStructSizeText = FText::GetEmpty();
+		SelectionSystemMemSizeText = FText::GetEmpty();
+		SelectionVideoMemSizeText = FText::GetEmpty();
+	}
+}
+
+////////////////////////////////////////////////////////////////////////////////////////////////////
+
+void SObjectTableTreeView::TreeView_OnSelectionChanged(FTableTreeNodePtr SelectedItem, ESelectInfo::Type SelectInfo)
+{
+	UpdateSelectionStatsText();
+}
+
+////////////////////////////////////////////////////////////////////////////////////////////////////
+
+void SObjectTableTreeView::TreeView_OnMouseButtonDoubleClick(FTableTreeNodePtr NodePtr)
+{
+	if (NodePtr->IsGroup())
+	{
+		const bool bIsGroupExpanded = TreeView->IsItemExpanded(NodePtr);
+		TreeView->SetItemExpansion(NodePtr, !bIsGroupExpanded);
+	}
+	else if (NodePtr->Is<FObjectNode>())
+	{
+		FObjectNode& ObjectNode = NodePtr->As<FObjectNode>();
+		//...
+	}
+}
+
+////////////////////////////////////////////////////////////////////////////////////////////////////
+// Grouping
+////////////////////////////////////////////////////////////////////////////////////////////////////
+
+void SObjectTableTreeView::InternalCreateGroupings()
+{
+	STableTreeView::InternalCreateGroupings();
+
+	AvailableGroupings.RemoveAll(
+		[](TSharedPtr<FTreeNodeGrouping>& Grouping)
+		{
+			if (Grouping->Is<FTreeNodeGroupingByUniqueValue>())
+			{
+				const FName ColumnId = Grouping->As<FTreeNodeGroupingByUniqueValue>().GetColumnId();
+				return (ColumnId == FObjectTableColumns::ObjectIdColumnId);
+			}
+			else if (Grouping->Is<FTreeNodeGroupingByPathBreakdown>())
+			{
+				const FName ColumnId = Grouping->As<FTreeNodeGroupingByPathBreakdown>().GetColumnId();
+				return (ColumnId == FObjectTableColumns::ObjectNameColumnId);
+			}
+			return false;
+		});
+
+	int32 Index = 1; // after the Flat ("All") grouping
+
+	AvailableGroupings.Insert(MakeShared<FObjectGroupingByOuter>(), Index++);
+	//AvailableGroupings.Insert(MakeShared<FUObjectByClassNodeGrouping>(SharedThis(this)), Index++);
+	//AvailableGroupings.Insert(MakeShared<FUObjectByPackageNodeGrouping>(SharedThis(this)), Index++);
+}
+
+////////////////////////////////////////////////////////////////////////////////////////////////////
+
+void SObjectTableTreeView::InitAvailableViewPresets()
+{
+	AvailableViewPresets.Add(FObjectTableViewPresets::CreateAssetViewPreset(*this));
+	AvailableViewPresets.Add(FObjectTableViewPresets::CreateObjectViewPreset(*this));
+	AvailableViewPresets.Add(FObjectTableViewPresets::CreateClassViewPreset(*this));
+	AvailableViewPresets.Add(FObjectTableViewPresets::CreateOuterViewPreset(*this));
+	AvailableViewPresets.Add(FObjectTableViewPresets::CreatePackageViewPreset(*this));
+
+	SelectedViewPreset = AvailableViewPresets[0];
+}
+
+////////////////////////////////////////////////////////////////////////////////////////////////////
+
+void SObjectTableTreeView::Reset()
+{
+	STableTreeView::Reset();
+}
+
+////////////////////////////////////////////////////////////////////////////////////////////////////
+
+void SObjectTableTreeView::Tick(const FGeometry& AllottedGeometry, const double InCurrentTime, const float InDeltaTime)
+{
+	STableTreeView::Tick(AllottedGeometry, InCurrentTime, InDeltaTime);
+
+	// We need to check if the list of UObject snapshots has changed.
+	// But, ensure we do not check too often.
+	static uint64 NextTimestamp = 0;
+	uint64 Time = FPlatformTime::Cycles64();
+	if (Time > NextTimestamp && !bIsUpdateRunning)
+	{
+		RebuildTree(false);
+
+		const int32 NumTags = GetTableRowNodes().Num();
+		// 1000 tags --> check each 200ms
+		// 10000 tags --> check each 1.1s
+		// 100000 tags --> check each 10.1s
+		const double WaitTimeSec = 0.1 + (double)NumTags / 10000.0;
+		const uint64 WaitTime = static_cast<uint64>(WaitTimeSec / FPlatformTime::GetSecondsPerCycle64());
+		NextTimestamp = Time + WaitTime;
+	}
+}
+
+////////////////////////////////////////////////////////////////////////////////////////////////////
+
+void SObjectTableTreeView::RebuildTree(bool bResync)
+{
+	FStopwatch Stopwatch;
+	Stopwatch.Start();
+
+	FStopwatch SyncStopwatch;
+	SyncStopwatch.Start();
+
+	bool bListHasChanged = false;
+
+	if (bResync)
+	{
+		ObjectNodes.Empty();
+		TableRowNodes.Empty();
+		bListHasChanged = true;
+	}
+
+	if (bResync || PendingSnapshotId != CurrentSnapshotId)
+	{
+		if (Session.IsValid())
+		{
+			const TraceServices::IObjectProvider* ObjectProvider = TraceServices::ReadObjectProvider(*Session.Get());
+			if (ObjectProvider)
+			{
+				TraceServices::FProviderReadScopeLock ProviderReadScope(*ObjectProvider);
+
+				const TraceServices::IObjectSnapshot* Snapshot = ObjectProvider->GetSnapshot(PendingSnapshotId);
+				if (Snapshot)
+				{
+					ObjectNodes.Empty();
+					TableRowNodes.Empty();
+					bListHasChanged = true;
+
+					CurrentSnapshotId = PendingSnapshotId;
+
+					const uint32 ObjectCount = Snapshot->GetNumObjects();
+					ObjectNodes.AddDefaulted(ObjectCount); // index == object id
+					uint32 ValidObjectCount = 0;
+
+					uint32 ClassId = uint32(-1);
+					uint32 PackageId = uint32(-1);
+
+					TSharedPtr<FObjectTable> ObjectTable = StaticCastSharedPtr<FObjectTable>(GetTable());
+					for (uint32 ObjectId = 0; ObjectId < ObjectCount; ++ObjectId)
+					{
+						const TraceServices::FObjectInfo* Object = Snapshot->GetObject(ObjectId);
+						if (Object && Object->Id == ObjectId)
+						{
+							TSharedPtr<FObjectNode> ObjectNode = MakeShared<FObjectNode>(ObjectTable, *Object);
+							ObjectNodes[ObjectId] = ObjectNode;
+							++ValidObjectCount;
+
+							if (FCString::Strcmp(Object->Name, TEXT("Class")) == 0)
+							{
+								ClassId = Object->Id;
+								ObjectNode->SetIsClass();
+							}
+							else if (FCString::Strcmp(Object->Name, TEXT("Package")) == 0)
+							{
+								PackageId = Object->Id;
+								ObjectNode->SetIsPackage();
+							}
+						}
+					}
+
+					TableRowNodes.Reserve(ValidObjectCount);
+
+					for (TSharedPtr<FObjectNode>& NodePtr : ObjectNodes)
+					{
+						if (NodePtr.IsValid())
+						{
+							// Resolve pointers to class, outer and package objects.
+							NodePtr->SetClass(GetObjectNode(NodePtr->GetClassId()));
+							NodePtr->SetOuter(GetObjectNode(NodePtr->GetOuterId()));
+							NodePtr->SetPackage(GetObjectNode(NodePtr->GetPackageId()));
+
+							TableRowNodes.Add(NodePtr);
+						}
+					}
+				}
+
+			}
+		}
+	}
+
+	SyncStopwatch.Stop();
+
+	if (bListHasChanged)
+	{
+		// Save selection.
+		TArray<FTableTreeNodePtr> SelectedItems;
+		TreeView->GetSelectedItems(SelectedItems);
+
+		UpdateTree();
+
+		TreeView->RebuildList();
+
+		// Restore selection.
+		if (SelectedItems.Num() > 0)
+		{
+			TreeView->ClearSelection();
+			for (FTableTreeNodePtr& NodePtr : SelectedItems)
+			{
+				if (!NodePtr->Is<FObjectNode>())
+				{
+					NodePtr.Reset();
+					continue;
+				}
+				NodePtr = GetObjectNode(NodePtr->As<FObjectNode>().GetObjectId());
+			}
+			SelectedItems.RemoveAll([](const FTableTreeNodePtr& NodePtr) { return !NodePtr.IsValid(); });
+			if (SelectedItems.Num() > 0)
+			{
+				TreeView->SetItemSelection(SelectedItems, true);
+				TreeView->RequestScrollIntoView(SelectedItems.Last());
+			}
+		}
+	}
+
+	Stopwatch.Stop();
+	const double TotalTime = Stopwatch.GetAccumulatedTime();
+	if (TotalTime > 0.01)
+	{
+		const double SyncTime = SyncStopwatch.GetAccumulatedTime();
+		UE_LOG(LogObjectProfiler, Log, TEXT("[UObj] Tree view rebuilt in %.4fs (sync: %.4fs + update: %.4fs) --> %d objects"),
+			TotalTime, SyncTime, TotalTime - SyncTime, TableRowNodes.Num());
+	}
+}
+
+////////////////////////////////////////////////////////////////////////////////////////////////////
+
+void SObjectTableTreeView::ExtendMenu(TSharedRef<FExtender> Extender)
+{
+	Extender->AddMenuExtension("Misc", EExtensionHook::Before, nullptr, FMenuExtensionDelegate::CreateSP(this, &SObjectTableTreeView::ExtendMenuBeforeMisc));
+	Extender->AddMenuExtension("Misc", EExtensionHook::After, nullptr, FMenuExtensionDelegate::CreateSP(this, &SObjectTableTreeView::ExtendMenuAfterMisc));
+}
+
+////////////////////////////////////////////////////////////////////////////////////////////////////
+
+void SObjectTableTreeView::ExtendMenuBeforeMisc(FMenuBuilder& MenuBuilder)
+{
+	MenuBuilder.BeginSection("UObject1", LOCTEXT("ContextMenu_Section_UObject1", "UObject 1"));
+	{
+	}
+	MenuBuilder.EndSection();
+}
+
+////////////////////////////////////////////////////////////////////////////////////////////////////
+
+void SObjectTableTreeView::ExtendMenuAfterMisc(FMenuBuilder& MenuBuilder)
+{
+	MenuBuilder.BeginSection("UObject2", LOCTEXT("ContextMenu_Section_UObject2", "UObject 2"));
+	{
+	}
+	MenuBuilder.EndSection();
+}
+
+////////////////////////////////////////////////////////////////////////////////////////////////////
+
+bool SObjectTableTreeView::HasCustomNodeFilter() const
+{
+	return bHideObjectsWithZeroStructureSize || bHideObjectsWithZeroMemorySize;
+}
+
+////////////////////////////////////////////////////////////////////////////////////////////////////
+
+bool SObjectTableTreeView::FilterNodeCustom(const FTableTreeNode& InNode) const
+{
+	if (InNode.Is<FObjectNode>())
+	{
+		const FObjectNode& ObjectNode = InNode.As<FObjectNode>();
+		if (bHideObjectsWithZeroStructureSize)
+		{
+			if (ObjectNode.GetStructureSize() == 0)
+			{
+				return false;
+			}
+		}
+		if (bHideObjectsWithZeroMemorySize)
+		{
+			if (ObjectNode.GetSystemMemorySize() == 0 && ObjectNode.GetVideoMemorySize() == 0)
+			{
+				return false;
+			}
+		}
+	}
+	return true;
+}
+
+////////////////////////////////////////////////////////////////////////////////////////////////////
+
+TSharedPtr<FObjectNode> SObjectTableTreeView::GetObjectNode(uint32 ObjectId)
+{
+	return (ObjectId < uint32(ObjectNodes.Num())) ? ObjectNodes[ObjectId] :  nullptr;
+}
+
+////////////////////////////////////////////////////////////////////////////////////////////////////
+
+void SObjectTableTreeView::SelectObjectNode(uint32 ObjectId)
+{
+	TSharedPtr<FObjectNode> NodePtr = GetObjectNode(ObjectId);
+	if (NodePtr)
+	{
+		TreeView->SetSelection(NodePtr);
+		TreeView->RequestScrollIntoView(NodePtr);
+	}
+}
+
+////////////////////////////////////////////////////////////////////////////////////////////////////
+////////////////////////////////////////////////////////////////////////////////////////////////////
+
+} // namespace UE::Insights::ObjectProfiler
+
+#undef LOCTEXT_NAMESPACE
diff --git a/Engine/Source/Developer/TraceInsights/Private/Insights/ObjectProfiler/Widgets/SObjectTableTreeView.h b/Engine/Source/Developer/TraceInsights/Private/Insights/ObjectProfiler/Widgets/SObjectTableTreeView.h
new file mode 100644
--- /dev/null
+++ b/Engine/Source/Developer/TraceInsights/Private/Insights/ObjectProfiler/Widgets/SObjectTableTreeView.h
@@ -0,0 +1,110 @@
+// Copyright Epic Games, Inc. All Rights Reserved.
+
+#pragma once
+
+#include "CoreTypes.h"
+
+// TraceInsights
+#include "Insights/Table/Widgets/SSessionTableTreeView.h"
+
+namespace UE::Insights::ObjectProfiler
+{
+
+class FObjectNode;
+class FObjectTable;
+
+////////////////////////////////////////////////////////////////////////////////////////////////////
+/**
+ * A custom widget used to display the list of objects from a UObject Snapshot.
+ */
+class SObjectTableTreeView : public SSessionTableTreeView
+{
+public:
+	/** Default constructor. */
+	SObjectTableTreeView();
+
+	/** Virtual destructor. */
+	virtual ~SObjectTableTreeView();
+
+	SLATE_BEGIN_ARGS(SObjectTableTreeView) {}
+	SLATE_END_ARGS()
+
+	/**
+	 * Construct this widget
+	 * @param InArgs - The declaration data for this widget
+	 */
+	void Construct(const FArguments& InArgs);
+
+	virtual void ConstructHeaderArea(TSharedRef<SVerticalBox> InHostBox) override;
+	virtual void ConstructFooterArea(TSharedRef<SVerticalBox> InHostBox) override;
+	virtual TSharedPtr<SWidget> ConstructFilterToolbar() override;
+	virtual TSharedPtr<SWidget> ConstructToolbar() override;
+	virtual TSharedPtr<SWidget> ConstructFooter() override;
+
+	virtual void Reset();
+
+	/**
+	 * Ticks this widget.  Override in derived classes, but always call the parent implementation.
+	 *
+	 * @param  AllottedGeometry The space allotted for this widget
+	 * @param  InCurrentTime  Current absolute real time
+	 * @param  InDeltaTime  Real time passed since last tick
+	 */
+	virtual void Tick(const FGeometry& AllottedGeometry, const double InCurrentTime, const float InDeltaTime) override;
+
+	/**
+	 * Rebuilds the tree (if necessary).
+	 * @param bResync - If true, it forces a resync
+	 */
+	virtual void RebuildTree(bool bResync);
+
+protected:
+	virtual void InternalCreateGroupings() override;
+
+	virtual void ExtendMenu(TSharedRef<FExtender> Extender) override;
+	void ExtendMenuBeforeMisc(FMenuBuilder& MenuBuilder);
+	void ExtendMenuAfterMisc(FMenuBuilder& MenuBuilder);
+
+	virtual bool HasCustomNodeFilter() const override;
+	virtual bool FilterNodeCustom(const FTableTreeNode& InNode) const override;
+
+	virtual void TreeView_OnMouseButtonDoubleClick(FTableTreeNodePtr NodePtr) override;
+	virtual void TreeView_OnSelectionChanged(FTableTreeNodePtr SelectedItem, ESelectInfo::Type SelectInfo) override;
+
+private:
+	void UpdateSelectionStatsText();
+	FText GetNumSelectedObjectsText() const { return NumSelectedObjectsText; }
+	FText GetSelectedObjectsText() const { return  SelectedObjectsText; }
+	FText GetSelectionStructSizeText() const { return SelectionStructSizeText; }
+	FText GetSelectionSystemMemSizeText() const { return  SelectionSystemMemSizeText; }
+	FText GetSelectionVideoMemSizeText() const { return  SelectionVideoMemSizeText; }
+
+	void InitAvailableViewPresets();
+
+	TSharedPtr<FObjectNode> GetObjectNode(uint32 ObjectId);
+	void SelectObjectNode(uint32 ObjectId);
+
+private:
+	uint32 CurrentSnapshotId = uint32(-1);
+	uint32 PendingSnapshotId = 0;
+
+	TArray<TSharedPtr<FObjectNode>> ObjectNodes;
+
+	//////////////////////////////////////////////////
+	// Filters
+
+	bool bHideObjectsWithZeroStructureSize = false;
+	bool bHideObjectsWithZeroMemorySize = false;
+
+	//////////////////////////////////////////////////
+
+	FText NumSelectedObjectsText;
+	FText SelectedObjectsText;
+	FText SelectionStructSizeText;
+	FText SelectionSystemMemSizeText;
+	FText SelectionVideoMemSizeText;
+};
+
+////////////////////////////////////////////////////////////////////////////////////////////////////
+
+} // namespace UE::Insights::ObjectProfiler
diff --git a/Engine/Source/Developer/TraceInsights/Private/Insights/TraceInsightsModule.cpp b/Engine/Source/Developer/TraceInsights/Private/Insights/TraceInsightsModule.cpp
--- a/Engine/Source/Developer/TraceInsights/Private/Insights/TraceInsightsModule.cpp
+++ b/Engine/Source/Developer/TraceInsights/Private/Insights/TraceInsightsModule.cpp
@@ -20,6 +20,7 @@
 
 // TraceServices
 #include "TraceServices/ITraceServicesModule.h"
+#include "TraceServices/Model/AnalysisSession.h"
 
 // TraceInsightsCore
 #include "InsightsCore/ITraceInsightsCoreModule.h"
@@ -34,6 +35,7 @@
 #include "Insights/Log.h"
 #include "Insights/MemoryProfiler/MemoryProfilerManager.h"
 #include "Insights/NetworkingProfiler/NetworkingProfilerManager.h"
+#include "Insights/ObjectProfiler/ObjectProfilerManager.h"
 #include "Insights/TaskGraphProfiler/TaskGraphProfilerManager.h"
 #include "Insights/Tests/InsightsTestRunner.h"
 #include "Insights/TimingProfiler/TimingProfilerManager.h"
@@ -79,6 +81,7 @@ void FTraceInsightsModule::StartupModule()
 	RegisterComponent(UE::Insights::TaskGraphProfiler::FTaskGraphProfilerManager::CreateInstance());
 	RegisterComponent(UE::Insights::ContextSwitches::FContextSwitchesProfilerManager::CreateInstance());
 	RegisterComponent(UE::Insights::CookProfiler::FCookProfilerManager::CreateInstance());
+	RegisterComponent(UE::Insights::ObjectProfiler::FObjectProfilerManager::CreateInstance());
 	RegisterComponent(Insights::FTableImportTool::CreateInstance());
 
 #if !WITH_EDITOR
