[Insights] Fixed comments for ObjectProfilerManager.

#rnx
#Synced-CL 51492753

diff --git a/Engine/Source/Developer/TraceInsights/Private/Insights/ObjectProfiler/ObjectProfilerManager.h b/Engine/Source/Developer/TraceInsights/Private/Insights/ObjectProfiler/ObjectProfilerManager.h
--- a/Engine/Source/Developer/TraceInsights/Private/Insights/ObjectProfiler/ObjectProfilerManager.h
+++ b/Engine/Source/Developer/TraceInsights/Private/Insights/ObjectProfiler/ObjectProfilerManager.h
@@ -31,26 +31,26 @@ struct FObjectProfilerTabs
 
 ////////////////////////////////////////////////////////////////////////////////////////////////////
 /**
- * This class manages the Cooking Profiler state and settings.
+ * This class manages the Object Profiler state and settings.
  */
 class FObjectProfilerManager : public TSharedFromThis<FObjectProfilerManager>, public IInsightsComponent
 {
 public:
-	/** Creates the Cook Profiler manager, only one instance can exist. */
+	/** Creates the Object Profiler manager, only one instance can exist. */
 	FObjectProfilerManager(TSharedRef<FUICommandList> InCommandList);
 
 	/** Virtual destructor. */
 	virtual ~FObjectProfilerManager();
 
-	/** Creates an instance of the CookProfiler manager. */
+	/** Creates an instance of the ObjectProfiler manager. */
 	static TSharedPtr<FObjectProfilerManager> CreateInstance();
 
 	/**
-	 * @return the global instance of the Cook Profiler manager.
+	 * @return the global instance of the Object Profiler manager.
 	 * This is an internal singleton and cannot be used outside TraceInsights.
 	 * For external use:
 	 *     IUnrealInsightsModule& Module = FModuleManager::Get().LoadModuleChecked<IUnrealInsightsModule>("TraceInsights");
-	 *     Module.GetCookProfilerManager();
+	 *     Module.GetObjectProfilerManager();
 	 */
 	static TSharedPtr<FObjectProfilerManager> Get();
 
@@ -91,7 +91,7 @@ private:
 	/** Handle to the registered OnTick. */
 	FTSTicker::FDelegateHandle OnTickHandle;
 
-	/** A shared pointer to the global instance of the Cook Profiler manager. */
+	/** A shared pointer to the global instance of the Object Profiler manager. */
 	static TSharedPtr<FObjectProfilerManager> Instance;
 
 	TWeakPtr<FTabManager> TimingTabManager;
