Memory View tool initial work on the Editor side.

#changelist validated
#Synced-CL 51256133

diff --git a/Engine/Source/Developer/TraceInsights/Private/Insights/TraceInsightsModule.cpp b/Engine/Source/Developer/TraceInsights/Private/Insights/TraceInsightsModule.cpp
--- a/Engine/Source/Developer/TraceInsights/Private/Insights/TraceInsightsModule.cpp
+++ b/Engine/Source/Developer/TraceInsights/Private/Insights/TraceInsightsModule.cpp
@@ -36,6 +36,7 @@
 #include "Insights/MemoryProfiler/MemoryProfilerManager.h"
 #include "Insights/NetworkingProfiler/NetworkingProfilerManager.h"
 #include "Insights/ObjectProfiler/ObjectProfilerManager.h"
+#include "Insights/ObjectProfiler/Widgets/SObjectTableTreeView.h"
 #include "Insights/TaskGraphProfiler/TaskGraphProfilerManager.h"
 #include "Insights/Tests/InsightsTestRunner.h"
 #include "Insights/TimingProfiler/TimingProfilerManager.h"
@@ -624,4 +625,11 @@ void FTraceInsightsModule::HandleCodeAccessorOpenFileFailed(const FString& Filen
 
 ////////////////////////////////////////////////////////////////////////////////////////////////////
 
+TSharedRef<SWidget> FTraceInsightsModule::CreateObjectTableTreeView() const
+{
+	return SNew(UE::Insights::ObjectProfiler::SObjectTableTreeView);
+}
+
+////////////////////////////////////////////////////////////////////////////////////////////////////
+
 } // namespace UE::Insights
diff --git a/Engine/Source/Developer/TraceInsights/Private/Insights/TraceInsightsModule.h b/Engine/Source/Developer/TraceInsights/Private/Insights/TraceInsightsModule.h
--- a/Engine/Source/Developer/TraceInsights/Private/Insights/TraceInsightsModule.h
+++ b/Engine/Source/Developer/TraceInsights/Private/Insights/TraceInsightsModule.h
@@ -92,6 +92,8 @@ public:
 	virtual void ScheduleCommand(const FString& InCmd) override;
 	virtual bool Exec(const TCHAR* Cmd, FOutputDevice& Ar) override;
 
+	virtual TSharedRef<SWidget> CreateObjectTableTreeView() const override;
+
 private:
 	void RegisterTabSpawners();
 	void UnregisterTabSpawners();
diff --git a/Engine/Source/Developer/TraceInsights/Public/Insights/IUnrealInsightsModule.h b/Engine/Source/Developer/TraceInsights/Public/Insights/IUnrealInsightsModule.h
--- a/Engine/Source/Developer/TraceInsights/Public/Insights/IUnrealInsightsModule.h
+++ b/Engine/Source/Developer/TraceInsights/Public/Insights/IUnrealInsightsModule.h
@@ -19,6 +19,7 @@
 #define UE_API TRACEINSIGHTS_API
 
 class FExtender;
+class SWidget;
 class FWorkspaceItem;
 
 namespace UE::Trace { class FStoreClient; }
@@ -382,6 +383,11 @@ public:
 
 	/** Execute command. */
 	virtual bool Exec(const TCHAR* Cmd, FOutputDevice& Ar) = 0;
+
+	/**
+	* Creates and returns an ObjectTableTreeView.
+	*/
+	virtual TSharedRef<SWidget> CreateObjectTableTreeView() const = 0;
 };
 
 ////////////////////////////////////////////////////////////////////////////////////////////////////
diff --git a/Engine/Source/Developer/TraceTools/Private/Widgets/SToggleTraceButton.cpp b/Engine/Source/Developer/TraceTools/Private/Widgets/SToggleTraceButton.cpp
--- a/Engine/Source/Developer/TraceTools/Private/Widgets/SToggleTraceButton.cpp
+++ b/Engine/Source/Developer/TraceTools/Private/Widgets/SToggleTraceButton.cpp
@@ -52,6 +52,7 @@ void SToggleTraceButton::Construct(const FArguments& InArgs)
 	DynamicToolTipTextAttribute.Assign(*this, InArgs._DynamicToolTipText);
 	LabelTextAttribute.Assign(*this, InArgs._LabelText);
 	ButtonSize = InArgs._ButtonSize;
+	bShowStopButton = InArgs._ShowStopButton;
 
 	ChildSlot
 	[
@@ -119,7 +120,7 @@ EVisibility SToggleTraceButton::GetStartTraceIconVisibility() const
 
 EVisibility SToggleTraceButton::GetStopTraceIconVisibility() const
 {
-	if (bIsTraceRecordButtonHovered && IsTraceRunningAttribute.Get())
+	if (bShowStopButton && bIsTraceRecordButtonHovered && IsTraceRunningAttribute.Get())
 	{
 		return EVisibility::Visible;
 	}
diff --git a/Engine/Source/Developer/TraceTools/Public/TraceTools/Widgets/SToggleTraceButton.h b/Engine/Source/Developer/TraceTools/Public/TraceTools/Widgets/SToggleTraceButton.h
--- a/Engine/Source/Developer/TraceTools/Public/TraceTools/Widgets/SToggleTraceButton.h
+++ b/Engine/Source/Developer/TraceTools/Public/TraceTools/Widgets/SToggleTraceButton.h
@@ -30,6 +30,7 @@ public:
 
 	SLATE_BEGIN_ARGS(SToggleTraceButton) 
 		: _ButtonSize(EButtonSize::StatusBar)
+		, _ShowStopButton(true)
 	{ }
 	
 	/* An event called when the toggle trace button is clicked. */
@@ -47,6 +48,9 @@ public:
 	/* Specifies the size of the button from a few presets. */
 	SLATE_ARGUMENT(EButtonSize, ButtonSize)
 
+	/* If true, the tracing icon will be replaced with a stop icon when hovered. */
+	SLATE_ARGUMENT(bool, ShowStopButton)
+
 	SLATE_END_ARGS()
 
 public:
@@ -80,6 +84,7 @@ private:
 	TSlateAttribute<FText> DynamicToolTipTextAttribute;
 	TSlateAttribute<FText> LabelTextAttribute;
 	EButtonSize ButtonSize;
+	bool bShowStopButton;
 };
 
 } // namespace UE::TraceTools
\ No newline at end of file
