[Insights] Fixed value formatter name (FInt64ValueFormatterAsUInt32InfiniteNumber).

#Synced-CL 51580401

diff --git a/Engine/Source/Developer/TraceInsights/Private/Insights/MemoryProfiler/ViewModels/MemAllocTable.cpp b/Engine/Source/Developer/TraceInsights/Private/Insights/MemoryProfiler/ViewModels/MemAllocTable.cpp
--- a/Engine/Source/Developer/TraceInsights/Private/Insights/MemoryProfiler/ViewModels/MemAllocTable.cpp
+++ b/Engine/Source/Developer/TraceInsights/Private/Insights/MemoryProfiler/ViewModels/MemAllocTable.cpp
@@ -215,7 +215,7 @@ void FMemAllocTable::AddDefaultColumns()
 		TSharedRef<ITableCellValueGetter> Getter = MakeShared<FMemAllocEndEventIndexValueGetter>();
 		Column.SetValueGetter(Getter);
 
-		TSharedRef<ITableCellValueFormatter> Formatter = MakeShared<FInt64ValueFormatterAsUInt32InfinteNumber>();
+		TSharedRef<ITableCellValueFormatter> Formatter = MakeShared<FInt64ValueFormatterAsUInt32InfiniteNumber>();
 		Column.SetValueFormatter(Formatter);
 
 		TSharedRef<ITableCellValueSorter> Sorter = MakeShared<FSorterByInt64Value>(ColumnRef);
@@ -264,7 +264,7 @@ void FMemAllocTable::AddDefaultColumns()
 		TSharedRef<ITableCellValueGetter> Getter = MakeShared<FMemAllocEventDistanceValueGetter>();
 		Column.SetValueGetter(Getter);
 
-		TSharedRef<ITableCellValueFormatter> Formatter = MakeShared<FInt64ValueFormatterAsUInt32InfinteNumber>();
+		TSharedRef<ITableCellValueFormatter> Formatter = MakeShared<FInt64ValueFormatterAsUInt32InfiniteNumber>();
 		Column.SetValueFormatter(Formatter);
 
 		TSharedRef<ITableCellValueSorter> Sorter = MakeShared<FSorterByInt64Value>(ColumnRef);
diff --git a/Engine/Source/Developer/TraceInsightsCore/Private/InsightsCore/Table/ViewModels/TableCellValueFormatter.cpp b/Engine/Source/Developer/TraceInsightsCore/Private/InsightsCore/Table/ViewModels/TableCellValueFormatter.cpp
--- a/Engine/Source/Developer/TraceInsightsCore/Private/InsightsCore/Table/ViewModels/TableCellValueFormatter.cpp
+++ b/Engine/Source/Developer/TraceInsightsCore/Private/InsightsCore/Table/ViewModels/TableCellValueFormatter.cpp
@@ -121,7 +121,7 @@ FText FBoolValueFormatterAsOnOff::FormatValue(const TOptional<FTableCellValue>&
 
 ////////////////////////////////////////////////////////////////////////////////////////////////////
 
-FText FInt64ValueFormatterAsUInt32InfinteNumber::FormatValue(const TOptional<FTableCellValue>& InValue) const
+FText FInt64ValueFormatterAsUInt32InfiniteNumber::FormatValue(const TOptional<FTableCellValue>& InValue) const
 {
 	if (InValue.IsSet())
 	{
@@ -137,7 +137,7 @@ FText FInt64ValueFormatterAsUInt32InfinteNumber::FormatValue(const TOptional<FTa
 
 ////////////////////////////////////////////////////////////////////////////////////////////////////
 
-FText FInt64ValueFormatterAsUInt32InfinteNumber::FormatValueForTooltip(const TOptional<FTableCellValue>& InValue) const
+FText FInt64ValueFormatterAsUInt32InfiniteNumber::FormatValueForTooltip(const TOptional<FTableCellValue>& InValue) const
 {
 	if (InValue.IsSet())
 	{
diff --git a/Engine/Source/Developer/TraceInsightsCore/Public/InsightsCore/Table/ViewModels/TableCellValueFormatter.h b/Engine/Source/Developer/TraceInsightsCore/Public/InsightsCore/Table/ViewModels/TableCellValueFormatter.h
--- a/Engine/Source/Developer/TraceInsightsCore/Public/InsightsCore/Table/ViewModels/TableCellValueFormatter.h
+++ b/Engine/Source/Developer/TraceInsightsCore/Public/InsightsCore/Table/ViewModels/TableCellValueFormatter.h
@@ -122,13 +122,16 @@ public:
 
 ////////////////////////////////////////////////////////////////////////////////////////////////////
 
-class FInt64ValueFormatterAsUInt32InfinteNumber : public FTableCellValueFormatter
+class FInt64ValueFormatterAsUInt32InfiniteNumber : public FTableCellValueFormatter
 {
 public:
 	UE_API virtual FText FormatValue(const TOptional<FTableCellValue>& InValue) const override;
 	UE_API virtual FText FormatValueForTooltip(const TOptional<FTableCellValue>& InValue) const override;
 };
 
+UE_DEPRECATED(5.8, "Please fix the spelling for 'Infinte'")
+typedef FInt64ValueFormatterAsUInt32InfiniteNumber FInt64ValueFormatterAsUInt32InfinteNumber;
+
 ////////////////////////////////////////////////////////////////////////////////////////////////////
 
 class FInt64ValueFormatterAsHex32 : public FTableCellValueFormatter
