[Insights]
* Object Profiler: Added grouping by object name (path breakdown) for the Assets preset.
* Object Profiler: Added custom FObjectNode class hierarchy for core object types (UField, UStruct, UClass and UPackage). Type-specific data (SuperId, InheritanceSuperId, StructureSize, PackageId, Path) moved into the appropriate subclasses.
* Object Profiler: Added new table columns for "Package Unique Id", "Package Path", "Super Id" and "Inheritance Super Id".

#rnx
#Synced-CL 51620660

diff --git a/Engine/Source/Developer/TraceInsights/Private/Insights/ObjectProfiler/ViewModels/ObjectGroupingByOuter.cpp b/Engine/Source/Developer/TraceInsights/Private/Insights/ObjectProfiler/ViewModels/ObjectGroupingByOuter.cpp
--- a/Engine/Source/Developer/TraceInsights/Private/Insights/ObjectProfiler/ViewModels/ObjectGroupingByOuter.cpp
+++ b/Engine/Source/Developer/TraceInsights/Private/Insights/ObjectProfiler/ViewModels/ObjectGroupingByOuter.cpp
@@ -11,7 +11,7 @@
 // TraceInsights
 #include "Insights/ObjectProfiler/ViewModels/ObjectNode.h"
 
-#define LOCTEXT_NAMESPACE "UE::Insights::MemoryProfiler::FMemAllocNode"
+#define LOCTEXT_NAMESPACE "UE::Insights::ObjectProfiler::FObjectNode"
 
 namespace UE::Insights::ObjectProfiler
 {
@@ -76,6 +76,29 @@ void FObjectGroupingByOuter::GroupNodes(
 
 ////////////////////////////////////////////////////////////////////////////////////////////////////
 
+static TSharedRef<FObjectNode> CloneObjectNode(FObjectNode& InNode)
+{
+	if (InNode.Is<FPackageObjectNode>())
+	{
+		return MakeShared<FPackageObjectNode>(InNode.As<FPackageObjectNode>());
+	}
+	if (InNode.Is<FClassObjectNode>())
+	{
+		return MakeShared<FClassObjectNode>(InNode.As<FClassObjectNode>());
+	}
+	if (InNode.Is<FStructObjectNode>())
+	{
+		return MakeShared<FStructObjectNode>(InNode.As<FStructObjectNode>());
+	}
+	if (InNode.Is<FFieldObjectNode>())
+	{
+		return MakeShared<FFieldObjectNode>(InNode.As<FFieldObjectNode>());
+	}
+	return MakeShared<FObjectNode>(InNode);
+}
+
+////////////////////////////////////////////////////////////////////////////////////////////////////
+
 TSharedPtr<FObjectNode> FObjectGroupingByOuter::AddObject(FObjectNode& ObjNode) const
 {
 	const uint32 ObjId = ObjNode.GetObjectId();
@@ -88,7 +111,7 @@ TSharedPtr<FObjectNode> FObjectGroupingByOuter::AddObject(FObjectNode& ObjNode)
 		return ObjectNodes[ObjId];
 	}
 
-	TSharedPtr<FObjectNode> NewNode = MakeShared<FObjectNode>(ObjNode);
+	TSharedPtr<FObjectNode> NewNode = CloneObjectNode(ObjNode);
 	NewNode->SetIsMissing();
 	ObjectNodes[ObjId] = NewNode;
 
diff --git a/Engine/Source/Developer/TraceInsights/Private/Insights/ObjectProfiler/ViewModels/ObjectGroupingByOuter.h b/Engine/Source/Developer/TraceInsights/Private/Insights/ObjectProfiler/ViewModels/ObjectGroupingByOuter.h
--- a/Engine/Source/Developer/TraceInsights/Private/Insights/ObjectProfiler/ViewModels/ObjectGroupingByOuter.h
+++ b/Engine/Source/Developer/TraceInsights/Private/Insights/ObjectProfiler/ViewModels/ObjectGroupingByOuter.h
@@ -28,7 +28,7 @@ private:
 	TSharedPtr<FObjectNode> AddObject(FObjectNode& ObjNode) const;
 	TSharedPtr<FCustomTableTreeNode> CreateObjectGroupNode(FObjectNode& ObjNode) const;
 
-	bool bCreateChildrenSubGroups = true;
+	bool bCreateChildrenSubGroups = false;
 	mutable FTableTreeNode* ParentGroupPtr = nullptr;
 	mutable TArray<TSharedPtr<FObjectNode>> ObjectNodes;
 };
diff --git a/Engine/Source/Developer/TraceInsights/Private/Insights/ObjectProfiler/ViewModels/ObjectNode.cpp b/Engine/Source/Developer/TraceInsights/Private/Insights/ObjectProfiler/ViewModels/ObjectNode.cpp
--- a/Engine/Source/Developer/TraceInsights/Private/Insights/ObjectProfiler/ViewModels/ObjectNode.cpp
+++ b/Engine/Source/Developer/TraceInsights/Private/Insights/ObjectProfiler/ViewModels/ObjectNode.cpp
@@ -13,7 +13,13 @@ namespace UE::Insights::ObjectProfiler
 {
 
 INSIGHTS_IMPLEMENT_RTTI(FObjectNode)
+INSIGHTS_IMPLEMENT_RTTI(FFieldObjectNode)
+INSIGHTS_IMPLEMENT_RTTI(FStructObjectNode)
+INSIGHTS_IMPLEMENT_RTTI(FClassObjectNode)
+INSIGHTS_IMPLEMENT_RTTI(FPackageObjectNode)
 
+////////////////////////////////////////////////////////////////////////////////////////////////////
+// FObjectNode
 ////////////////////////////////////////////////////////////////////////////////////////////////////
 
 FObjectNode::FObjectNode(TWeakPtr<FObjectTable> InParentTable, const TraceServices::FObjectInfo& ObjectInfo)
@@ -21,31 +27,14 @@ FObjectNode::FObjectNode(TWeakPtr<FObjectTable> InParentTable, const TraceServic
 	, ObjectId(ObjectInfo.Id)
 	, ClassId(ObjectInfo.ClassId)
 	, OuterId(ObjectInfo.OuterId)
+	, ObjectFlags(ObjectInfo.Flags)
 	, ObjectName(ObjectInfo.Name)
 	, VersePath(ObjectInfo.VersePath)
-	, ObjectFlags(ObjectInfo.Flags)
-	, StructureSize(ObjectInfo.StructureSize)
 	, SystemMemorySize(int64(ObjectInfo.SystemMemoryBytes))
 	, VideoMemorySize(int64(ObjectInfo.VideoMemoryBytes))
 	, InclSystemMemorySize(int64(ObjectInfo.TotalSystemMemoryBytes))
 	, InclVideoMemorySize(int64(ObjectInfo.TotalVideoMemoryBytes))
 {
-	if (EnumHasAnyFlags(ObjectInfo.FlagsEx, TraceServices::EObjectInfoFlags::IsField))
-	{
-		SetIsField();
-	}
-	if (EnumHasAnyFlags(ObjectInfo.FlagsEx, TraceServices::EObjectInfoFlags::IsStruct))
-	{
-		SetIsStruct();
-	}
-	if (EnumHasAnyFlags(ObjectInfo.FlagsEx, TraceServices::EObjectInfoFlags::IsClass))
-	{
-		SetIsClass();
-	}
-	if (EnumHasAnyFlags(ObjectInfo.FlagsEx, TraceServices::EObjectInfoFlags::IsPackage))
-	{
-		SetIsPackage();
-	}
 }
 
 ////////////////////////////////////////////////////////////////////////////////////////////////////
@@ -55,13 +44,12 @@ FObjectNode::FObjectNode(const FObjectNode& Clone)
 	, ObjectId(Clone.ObjectId)
 	, ClassId(Clone.ClassId)
 	, OuterId(Clone.OuterId)
+	, ObjectFlags(Clone.ObjectFlags)
 	, Class(Clone.Class)
 	, Outer(Clone.Outer)
 	, ObjectName(Clone.ObjectName)
 	, VersePath(Clone.VersePath)
 	, ObjectPath(Clone.ObjectPath)
-	, ObjectFlags(Clone.ObjectFlags)
-	, StructureSize(Clone.StructureSize)
 	, SystemMemorySize(Clone.SystemMemorySize)
 	, VideoMemorySize(Clone.VideoMemorySize)
 	, InclSystemMemorySize(Clone.InclSystemMemorySize)
@@ -109,25 +97,6 @@ const FText FObjectNode::GetTooltipText() const
 
 const FSlateBrush* FObjectNode::GetIcon() const
 {
-	if (InternalFlags != EInternalFlags::None)
-	{
-		if (IsClass())
-		{
-			return FInsightsStyle::GetBrush("Icons.UClass.TreeItem");
-		}
-		else if (IsStruct())
-		{
-			return FInsightsStyle::GetBrush("Icons.UStruct.TreeItem");
-		}
-		else if (IsField())
-		{
-			return FInsightsStyle::GetBrush("Icons.UField.TreeItem");
-		}
-		else if (IsPackage())
-		{
-			return FInsightsStyle::GetBrush("Icons.UPackage.TreeItem");
-		}
-	}
 	return FInsightsStyle::GetBrush("Icons.UObject.TreeItem");
 }
 
@@ -135,25 +104,6 @@ const FSlateBrush* FObjectNode::GetIcon() const
 
 FLinearColor FObjectNode::GetIconColor() const
 {
-	if (InternalFlags != EInternalFlags::None)
-	{
-		if (IsClass())
-		{
-			return FLinearColor(0.5f, 1.0f, 1.0f, 1.0f);
-		}
-		else if (IsStruct())
-		{
-			return FLinearColor(0.6f, 1.0f, 1.0f, 1.0f);
-		}
-		else if (IsField())
-		{
-			return FLinearColor(0.7f, 1.0f, 1.0f, 1.0f);
-		}
-		else if (IsPackage())
-		{
-			return FLinearColor(1.0f, 1.0f, 0.5f, 1.0f);
-		}
-	}
 	return GetDefaultColor(IsGroup());
 }
 
@@ -187,6 +137,140 @@ TSharedPtr<FObjectNode> FObjectNode::GetPackage() const
 	return nullptr;
 }
 
+////////////////////////////////////////////////////////////////////////////////////////////////////
+// FFieldObjectNode
+////////////////////////////////////////////////////////////////////////////////////////////////////
+
+FFieldObjectNode::FFieldObjectNode(TWeakPtr<FObjectTable> InParentTable, const TraceServices::FObjectInfo& ObjectInfo)
+	: FObjectNode(InParentTable, ObjectInfo)
+{
+	ensure(EnumHasAnyFlags(ObjectInfo.FlagsEx, TraceServices::EObjectInfoFlags::IsField));
+}
+
+////////////////////////////////////////////////////////////////////////////////////////////////////
+
+FFieldObjectNode::FFieldObjectNode(const FFieldObjectNode& Clone)
+	: FObjectNode(Clone)
+{
+}
+
+////////////////////////////////////////////////////////////////////////////////////////////////////
+
+const FSlateBrush* FFieldObjectNode::GetIcon() const
+{
+	return FInsightsStyle::GetBrush("Icons.UField.TreeItem");
+}
+
+////////////////////////////////////////////////////////////////////////////////////////////////////
+
+FLinearColor FFieldObjectNode::GetIconColor() const
+{
+	return FLinearColor(0.7f, 1.0f, 1.0f, 1.0f);
+}
+
+////////////////////////////////////////////////////////////////////////////////////////////////////
+// FStructObjectNode
+////////////////////////////////////////////////////////////////////////////////////////////////////
+
+FStructObjectNode::FStructObjectNode(TWeakPtr<FObjectTable> InParentTable, const TraceServices::FObjectInfo& ObjectInfo)
+	: FFieldObjectNode(InParentTable, ObjectInfo)
+	, SuperId(ObjectInfo.SuperId)
+	, InheritanceSuperId(ObjectInfo.InheritanceSuperId)
+	, StructureSize(ObjectInfo.StructureSize)
+{
+	ensure(EnumHasAnyFlags(ObjectInfo.FlagsEx, TraceServices::EObjectInfoFlags::IsStruct));
+}
+
+////////////////////////////////////////////////////////////////////////////////////////////////////
+
+FStructObjectNode::FStructObjectNode(const FStructObjectNode& Clone)
+	: FFieldObjectNode(Clone)
+	, SuperId(Clone.SuperId)
+	, InheritanceSuperId(Clone.InheritanceSuperId)
+	, StructureSize(Clone.StructureSize)
+{
+}
+
+////////////////////////////////////////////////////////////////////////////////////////////////////
+
+const FSlateBrush* FStructObjectNode::GetIcon() const
+{
+	return FInsightsStyle::GetBrush("Icons.UStruct.TreeItem");
+}
+
+////////////////////////////////////////////////////////////////////////////////////////////////////
+
+FLinearColor FStructObjectNode::GetIconColor() const
+{
+	return FLinearColor(0.6f, 1.0f, 1.0f, 1.0f);
+}
+
+////////////////////////////////////////////////////////////////////////////////////////////////////
+// FClassObjectNode
+////////////////////////////////////////////////////////////////////////////////////////////////////
+
+FClassObjectNode::FClassObjectNode(TWeakPtr<FObjectTable> InParentTable, const TraceServices::FObjectInfo& ObjectInfo)
+	: FStructObjectNode(InParentTable, ObjectInfo)
+{
+	ensure(EnumHasAnyFlags(ObjectInfo.FlagsEx, TraceServices::EObjectInfoFlags::IsClass));
+}
+
+////////////////////////////////////////////////////////////////////////////////////////////////////
+
+FClassObjectNode::FClassObjectNode(const FClassObjectNode& Clone)
+	: FStructObjectNode(Clone)
+{
+}
+
+////////////////////////////////////////////////////////////////////////////////////////////////////
+
+const FSlateBrush* FClassObjectNode::GetIcon() const
+{
+	return FInsightsStyle::GetBrush("Icons.UClass.TreeItem");
+}
+
+////////////////////////////////////////////////////////////////////////////////////////////////////
+
+FLinearColor FClassObjectNode::GetIconColor() const
+{
+	return FLinearColor(0.5f, 1.0f, 1.0f, 1.0f);
+}
+
+////////////////////////////////////////////////////////////////////////////////////////////////////
+// FPackageObjectNode
+////////////////////////////////////////////////////////////////////////////////////////////////////
+
+FPackageObjectNode::FPackageObjectNode(TWeakPtr<FObjectTable> InParentTable, const TraceServices::FObjectInfo& ObjectInfo)
+	: FObjectNode(InParentTable, ObjectInfo)
+	, PackageId(ObjectInfo.PackageId)
+	, Path(ObjectInfo.PackagePath)
+{
+	ensure(EnumHasAnyFlags(ObjectInfo.FlagsEx, TraceServices::EObjectInfoFlags::IsPackage));
+}
+
+////////////////////////////////////////////////////////////////////////////////////////////////////
+
+FPackageObjectNode::FPackageObjectNode(const FPackageObjectNode& Clone)
+	: FObjectNode(Clone)
+	, PackageId(Clone.PackageId)
+	, Path(Clone.Path)
+{
+}
+
+////////////////////////////////////////////////////////////////////////////////////////////////////
+
+const FSlateBrush* FPackageObjectNode::GetIcon() const
+{
+	return FInsightsStyle::GetBrush("Icons.UPackage.TreeItem");
+}
+
+////////////////////////////////////////////////////////////////////////////////////////////////////
+
+FLinearColor FPackageObjectNode::GetIconColor() const
+{
+	return FLinearColor(1.0f, 1.0f, 0.5f, 1.0f);
+}
+
 ////////////////////////////////////////////////////////////////////////////////////////////////////
 
 } // namespace UE::Insights::ObjectProfiler
diff --git a/Engine/Source/Developer/TraceInsights/Private/Insights/ObjectProfiler/ViewModels/ObjectNode.h b/Engine/Source/Developer/TraceInsights/Private/Insights/ObjectProfiler/ViewModels/ObjectNode.h
--- a/Engine/Source/Developer/TraceInsights/Private/Insights/ObjectProfiler/ViewModels/ObjectNode.h
+++ b/Engine/Source/Developer/TraceInsights/Private/Insights/ObjectProfiler/ViewModels/ObjectNode.h
@@ -39,11 +39,7 @@ private:
 	enum class EInternalFlags
 	{
 		None = 0,
-		IsField = 1 << 0,
-		IsStruct = 1 << 1,
-		IsClass = 1 << 2,
-		IsPackage = 1 << 3,
-		IsMissing = 1 << 4, // missing intermediate node in the Outer hierarchy
+		IsMissing = 1 << 0, // missing intermediate node in the Outer hierarchy
 	};
 
 public:
@@ -69,6 +65,8 @@ public:
 	uint32 GetClassId() const { return ClassId; }
 	uint32 GetOuterId() const { return OuterId; }
 
+	uint32 GetObjectFlags() const { return ObjectFlags; }
+
 	TSharedPtr<FObjectNode> GetClass() const { return Class.Pin(); }
 	TSharedPtr<FObjectNode> GetOuter() const { return Outer.Pin(); }
 
@@ -78,9 +76,7 @@ public:
 	const TCHAR* GetObjectPath() const { return *ObjectPath; }
 	const TCHAR* GetVersePath() const { return VersePath; }
 
-	uint32 GetObjectFlags() const { return ObjectFlags; }
-
-	int32 GetStructureSize() const { return StructureSize; }
+	virtual int32 GetStructureSize() const { return 0; }
 
 	int64 GetSystemMemorySize() const { return SystemMemorySize; }
 	int64 GetVideoMemorySize() const { return VideoMemorySize; }
@@ -90,19 +86,16 @@ public:
 	int64 GetInclusiveVideoMemorySize() const { return InclVideoMemorySize; }
 	int64 GetInclusiveEstimatedSize() const { return InclSystemMemorySize + InclVideoMemorySize; }
 
-	uint32 GetNumReferences() const { return 0; } // TODO
+	virtual uint32 GetNumReferences() const { return 0; }
 
-	bool IsField() const { return (uint32(InternalFlags) & uint32(EInternalFlags::IsField)) != 0; }
-	bool IsStruct() const { return (uint32(InternalFlags) & uint32(EInternalFlags::IsStruct)) != 0; }
-	bool IsClass() const { return (uint32(InternalFlags) & uint32(EInternalFlags::IsClass)) != 0; }
-	bool IsPackage() const { return (uint32(InternalFlags) & uint32(EInternalFlags::IsPackage)) != 0; }
 	bool IsMissing() const { return (uint32(InternalFlags) & uint32(EInternalFlags::IsMissing)) != 0; }
 
-private:
-	void SetIsField() { InternalFlags = EInternalFlags(uint32(InternalFlags) | uint32(EInternalFlags::IsField)); }
-	void SetIsStruct() { InternalFlags = EInternalFlags(uint32(InternalFlags) | uint32(EInternalFlags::IsStruct)); }
-	void SetIsClass() { InternalFlags = EInternalFlags(uint32(InternalFlags) | uint32(EInternalFlags::IsClass)); }
-	void SetIsPackage() { InternalFlags = EInternalFlags(uint32(InternalFlags) | uint32(EInternalFlags::IsPackage)); }
+	virtual bool IsField() const { return false; }
+	virtual bool IsStruct() const { return false; }
+	virtual bool IsClass() const { return false; }
+	virtual bool IsPackage() const { return false; }
+
+protected:
 	void SetIsMissing() { InternalFlags = EInternalFlags(uint32(InternalFlags) | uint32(EInternalFlags::IsMissing)); }
 	void ResetIsMissing() { InternalFlags = EInternalFlags(uint32(InternalFlags) & ~uint32(EInternalFlags::IsMissing)); }
 
@@ -115,16 +108,15 @@ private:
 	uint32 ClassId = uint32(-1);
 	uint32 OuterId = uint32(-1);
 
+	uint32 ObjectFlags = 0;
+
 	TWeakPtr<FObjectNode> Class;
 	TWeakPtr<FObjectNode> Outer;
 
-	const TCHAR* ObjectName = nullptr;
-	const TCHAR* VersePath = nullptr;
+	const TCHAR* ObjectName = nullptr; // persistent string (session lifetime)
+	const TCHAR* VersePath = nullptr; // persistent string (session lifetime)
 	FString ObjectPath;
 
-	uint32 ObjectFlags = 0;
-
-	int32 StructureSize = 0; // for UStruct and UClass
 	int64 SystemMemorySize = 0;
 	int64 VideoMemorySize = 0;
 	int64 InclSystemMemorySize = 0;
@@ -135,4 +127,100 @@ private:
 
 ////////////////////////////////////////////////////////////////////////////////////////////////////
 
+class FFieldObjectNode : public FObjectNode
+{
+	friend class SObjectTableTreeView;
+
+	INSIGHTS_DECLARE_RTTI(FFieldObjectNode, FObjectNode)
+
+public:
+	/** Initialization constructor for the UField node. */
+	explicit FFieldObjectNode(TWeakPtr<FObjectTable> InParentTable, const TraceServices::FObjectInfo& ObjectInfo);
+
+	explicit FFieldObjectNode(const FFieldObjectNode& Clone);
+
+	virtual const FSlateBrush* GetIcon() const override;
+	virtual FLinearColor GetIconColor() const override;
+
+	virtual bool IsField() const override { return true; }
+};
+
+////////////////////////////////////////////////////////////////////////////////////////////////////
+
+class FStructObjectNode : public FFieldObjectNode
+{
+	friend class SObjectTableTreeView;
+
+	INSIGHTS_DECLARE_RTTI(FStructObjectNode, FFieldObjectNode)
+
+public:
+	/** Initialization constructor for the UStruct node. */
+	explicit FStructObjectNode(TWeakPtr<FObjectTable> InParentTable, const TraceServices::FObjectInfo& ObjectInfo);
+
+	explicit FStructObjectNode(const FStructObjectNode& Clone);
+
+	virtual const FSlateBrush* GetIcon() const override;
+	virtual FLinearColor GetIconColor() const override;
+
+	virtual bool IsStruct() const override { return true; }
+
+	uint32 GetSuperId() const { return SuperId; }
+	uint32 GetInheritanceSuperId() const { return InheritanceSuperId; }
+	virtual int32 GetStructureSize() const override { return StructureSize; }
+
+private:
+	uint32 SuperId = uint32(-1);
+	uint32 InheritanceSuperId = uint32(-1);
+	int32 StructureSize = 0;
+};
+
+////////////////////////////////////////////////////////////////////////////////////////////////////
+
+class FClassObjectNode : public FStructObjectNode
+{
+	friend class SObjectTableTreeView;
+
+	INSIGHTS_DECLARE_RTTI(FClassObjectNode, FStructObjectNode)
+
+	virtual const FSlateBrush* GetIcon() const override;
+	virtual FLinearColor GetIconColor() const override;
+
+	virtual bool IsClass() const override { return true; }
+
+public:
+	/** Initialization constructor for the UClass node. */
+	explicit FClassObjectNode(TWeakPtr<FObjectTable> InParentTable, const TraceServices::FObjectInfo& ObjectInfo);
+
+	explicit FClassObjectNode(const FClassObjectNode& Clone);
+};
+
+////////////////////////////////////////////////////////////////////////////////////////////////////
+
+class FPackageObjectNode : public FObjectNode
+{
+	friend class SObjectTableTreeView;
+
+	INSIGHTS_DECLARE_RTTI(FPackageObjectNode, FObjectNode)
+
+public:
+	/** Initialization constructor for the UPackage node. */
+	explicit FPackageObjectNode(TWeakPtr<FObjectTable> InParentTable, const TraceServices::FObjectInfo& ObjectInfo);
+
+	explicit FPackageObjectNode(const FPackageObjectNode& Clone);
+
+	virtual const FSlateBrush* GetIcon() const override;
+	virtual FLinearColor GetIconColor() const override;
+
+	virtual bool IsPackage() const override { return true; }
+
+	uint64 GetPackageId() const { return PackageId; }
+	const TCHAR* GetPath() const { return Path; }
+
+private:
+	uint64 PackageId = 0;
+	const TCHAR* Path = nullptr; // persistent string (session lifetime)
+};
+
+////////////////////////////////////////////////////////////////////////////////////////////////////
+
 } // namespace UE::Insights::ObjectProfiler
diff --git a/Engine/Source/Developer/TraceInsights/Private/Insights/ObjectProfiler/ViewModels/ObjectTable.cpp b/Engine/Source/Developer/TraceInsights/Private/Insights/ObjectProfiler/ViewModels/ObjectTable.cpp
--- a/Engine/Source/Developer/TraceInsights/Private/Insights/ObjectProfiler/ViewModels/ObjectTable.cpp
+++ b/Engine/Source/Developer/TraceInsights/Private/Insights/ObjectProfiler/ViewModels/ObjectTable.cpp
@@ -28,17 +28,27 @@
 const FName FObjectTableColumns::OuterIdColumnId("OuterId");
 const FName FObjectTableColumns::PackageNameColumnId("PackageName");
 const FName FObjectTableColumns::PackageIdColumnId("PackageId");
+
+const FName FObjectTableColumns::PackageUniqueIdColumnId("PackageUniqueId");
+const FName FObjectTableColumns::PackagePathColumnId("PackagePath");
+
 const FName FObjectTableColumns::ObjectFlagsColumnId("ObjectFlags");
+
 const FName FObjectTableColumns::ObjectPathColumnId("ObjectPath");
 const FName FObjectTableColumns::VersePathColumnId("VersePath");
+
+const FName FObjectTableColumns::SuperIdColumnId("SuperId");
+const FName FObjectTableColumns::InheritanceSuperIdColumnId("InheritanceSuperId");
 const FName FObjectTableColumns::StructureSizeColumnId("StructureSize");
+
 const FName FObjectTableColumns::SystemMemorySizeColumnId("SystemMemorySize");
 const FName FObjectTableColumns::InclSystemMemorySizeColumnId("InclSystemMemorySize");
 const FName FObjectTableColumns::VideoMemorySizeColumnId("VideoMemorySize");
 const FName FObjectTableColumns::InclVideoMemorySizeColumnId("InclVideoMemorySize");
 const FName FObjectTableColumns::EstimatedSizeColumnId("EstimatedSize");
-const FName FObjectTableColumns::InclEstimatedSizeColumnId("InclEstimatedSize");
 const FName FObjectTableColumns::ImpactColumnId("Impact");
+const FName FObjectTableColumns::InclEstimatedSizeColumnId("InclEstimatedSize");
+
 const FName FObjectTableColumns::ReferencesColumnId("References");
 
 ////////////////////////////////////////////////////////////////////////////////////////////////////
@@ -259,7 +269,7 @@
 		//TSharedRef<IFilterValueConverter> Converter = MakeShared<FMemoryFilterValueConverter>();
 		//Column.SetValueConverter(Converter);
 
-		//Column.SetAggregation(ETableColumnAggregation::None);
+		Column.SetAggregation(ETableColumnAggregation::None);
 
 		AddColumn(ColumnRef);
 	}
@@ -525,7 +535,7 @@
 		AddColumn(ColumnRef);
 	}
 	//////////////////////////////////////////////////
-	// Package Id Column
+	// Outer Package Object Id Column
 	{
 		TSharedRef<FTableColumn> ColumnRef = MakeShared<FTableColumn>(FObjectTableColumns::PackageIdColumnId);
 		FTableColumn& Column = *ColumnRef;
@@ -533,8 +543,8 @@
 		Column.SetIndex(ColumnIndex++);
 
 		Column.SetShortName(LOCTEXT("PackageIdColumnName", "Package Id"));
-		Column.SetTitleName(LOCTEXT("PackageIdColumnTitle", "Package Id"));
-		Column.SetDescription(LOCTEXT("PackageIdColumnDesc", "The id of the UPackage of the current UObject instance"));
+		Column.SetTitleName(LOCTEXT("PackageIdColumnTitle", "Package Object Id"));
+		Column.SetDescription(LOCTEXT("PackageIdColumnDesc", "The object id of the outer package of the current object instance"));
 
 		Column.SetFlags(ETableColumnFlags::CanBeHidden |
 						ETableColumnFlags::CanBeFiltered);
@@ -579,6 +589,116 @@
 		AddColumn(ColumnRef);
 	}
 	//////////////////////////////////////////////////
+	// Package Unique Id Column (UPackage)
+	{
+		TSharedRef<FTableColumn> ColumnRef = MakeShared<FTableColumn>(FObjectTableColumns::PackageUniqueIdColumnId);
+		FTableColumn& Column = *ColumnRef;
+
+		Column.SetIndex(ColumnIndex++);
+
+		Column.SetShortName(LOCTEXT("PackageUniqueIdColumnName", "Package UId"));
+		Column.SetTitleName(LOCTEXT("PackageUniqueIdColumnTitle", "Package Unique Id"));
+		Column.SetDescription(LOCTEXT("PackageUniqueIdColumnDesc", "The unique id of a UPackage object"));
+
+		Column.SetFlags(ETableColumnFlags::CanBeHidden |
+						ETableColumnFlags::CanBeFiltered);
+
+		Column.SetHorizontalAlignment(HAlign_Left);
+		Column.SetInitialWidth(60.0f);
+
+		Column.SetDataType(ETableCellDataType::Int64);
+
+		class FValueGetter : public FAggregatedValueGetter
+		{
+		public:
+			virtual const TOptional<FTableCellValue> GetNodeValue(const FBaseTreeNode& Node) const override
+			{
+				if (Node.Is<FPackageObjectNode>())
+				{
+					const FPackageObjectNode& PackageNode = Node.As<FPackageObjectNode>();
+					const uint64 PackageUniqueId = PackageNode.GetPackageId();
+					return FTableCellValue(int64(PackageUniqueId));
+				}
+				return TOptional<FTableCellValue>();
+			}
+		};
+		TSharedRef<ITableCellValueGetter> Getter = MakeShared<FValueGetter>();
+		Column.SetValueGetter(Getter);
+
+		class FValueFormatter : public FTableCellValueFormatter
+		{
+		public:
+			virtual FText FormatValue(const TOptional<FTableCellValue>& InValue) const override
+			{
+				return InValue.IsSet() ?
+					FText::FromString(FString::Printf(TEXT("0x%llX"), static_cast<uint64>(InValue.GetValue().Int64))) :
+					FText::GetEmpty();
+			}
+		};
+		TSharedRef<ITableCellValueFormatter> Formatter = MakeShared<FValueFormatter>();
+		Column.SetValueFormatter(Formatter);
+
+		TSharedRef<ITableCellValueSorter> Sorter = MakeShared<FSorterByInt64Value>(ColumnRef);
+		Column.SetValueSorter(Sorter);
+		Column.SetInitialSortMode(EColumnSortMode::Descending);
+
+		//TSharedRef<IFilterValueConverter> Converter = MakeShared<FMemoryFilterValueConverter>();
+		//Column.SetValueConverter(Converter);
+
+		Column.SetAggregation(ETableColumnAggregation::None);
+
+		AddColumn(ColumnRef);
+	}
+	//////////////////////////////////////////////////
+	// Package Path Column (UPackage)
+	{
+		TSharedRef<FTableColumn> ColumnRef = MakeShared<FTableColumn>(FObjectTableColumns::PackagePathColumnId);
+		FTableColumn& Column = *ColumnRef;
+
+		Column.SetIndex(ColumnIndex++);
+
+		Column.SetShortName(LOCTEXT("PackagePathColumnName", "Package Path"));
+		Column.SetTitleName(LOCTEXT("PackagePathColumnTitle", "Package Path"));
+		Column.SetDescription(LOCTEXT("PackagePathColumnDesc", "The path of a UPackage object"));
+
+		Column.SetFlags(ETableColumnFlags::CanBeHidden |
+						ETableColumnFlags::CanBeFiltered);
+
+		Column.SetHorizontalAlignment(HAlign_Left);
+		Column.SetInitialWidth(120.0f);
+
+		Column.SetDataType(ETableCellDataType::CString);
+
+		class FValueGetter : public FAggregatedValueGetter
+		{
+		public:
+			virtual const TOptional<FTableCellValue> GetNodeValue(const FBaseTreeNode& Node) const override
+			{
+				if (Node.Is<FPackageObjectNode>())
+				{
+					const FPackageObjectNode& PackageNode = Node.As<FPackageObjectNode>();
+					if (PackageNode.GetPath())
+					{
+						return FTableCellValue(PackageNode.GetPath());
+					}
+				}
+				return TOptional<FTableCellValue>();
+			}
+		};
+		TSharedRef<ITableCellValueGetter> Getter = MakeShared<FValueGetter>();
+		Column.SetValueGetter(Getter);
+
+		TSharedRef<ITableCellValueFormatter> Formatter = MakeShared<FCStringValueFormatterAsText>();
+		Column.SetValueFormatter(Formatter);
+
+		TSharedRef<ITableCellValueSorter> Sorter = MakeShared<FSorterByCStringValue>(ColumnRef);
+		Column.SetValueSorter(Sorter);
+
+		Column.SetAggregation(ETableColumnAggregation::SameValue);
+
+		AddColumn(ColumnRef);
+	}
+	//////////////////////////////////////////////////
 	// Object Flags Column
 	{
 		TSharedRef<FTableColumn> ColumnRef = MakeShared<FTableColumn>(FObjectTableColumns::ObjectFlagsColumnId);
@@ -732,6 +852,114 @@
 		AddColumn(ColumnRef);
 	}
 	//////////////////////////////////////////////////
+	// Super Id Column (UStruct and UClass)
+	{
+		TSharedRef<FTableColumn> ColumnRef = MakeShared<FTableColumn>(FObjectTableColumns::SuperIdColumnId);
+		FTableColumn& Column = *ColumnRef;
+
+		Column.SetIndex(ColumnIndex++);
+
+		Column.SetShortName(LOCTEXT("SuperIdColumnName", "Super Id"));
+		Column.SetTitleName(LOCTEXT("SuperIdColumnTitle", "Super Id"));
+		Column.SetDescription(LOCTEXT("SuperIdColumnDesc", "The id of the super class/struct"));
+
+		Column.SetFlags(ETableColumnFlags::CanBeHidden |
+						ETableColumnFlags::CanBeFiltered);
+
+		Column.SetHorizontalAlignment(HAlign_Left);
+		Column.SetInitialWidth(60.0f);
+
+		Column.SetDataType(ETableCellDataType::Int64);
+
+		class FValueGetter : public FAggregatedValueGetter
+		{
+		public:
+			virtual const TOptional<FTableCellValue> GetNodeValue(const FBaseTreeNode& Node) const override
+			{
+				if (Node.Is<FStructObjectNode>())
+				{
+					const FStructObjectNode& StructNode = Node.As<FStructObjectNode>();
+					const uint32 Id = StructNode.GetSuperId();
+					if (Id != uint32(-1))
+					{
+						return FTableCellValue(int64(Id));
+					}
+				}
+				return TOptional<FTableCellValue>();
+			}
+		};
+		TSharedRef<ITableCellValueGetter> Getter = MakeShared<FValueGetter>();
+		Column.SetValueGetter(Getter);
+
+		TSharedRef<ITableCellValueFormatter> Formatter = MakeShared<FInt64ValueFormatterAsNumber>();
+		Column.SetValueFormatter(Formatter);
+
+		TSharedRef<ITableCellValueSorter> Sorter = MakeShared<FSorterByInt64Value>(ColumnRef);
+		Column.SetValueSorter(Sorter);
+		Column.SetInitialSortMode(EColumnSortMode::Descending);
+
+		//TSharedRef<IFilterValueConverter> Converter = MakeShared<FMemoryFilterValueConverter>();
+		//Column.SetValueConverter(Converter);
+
+		Column.SetAggregation(ETableColumnAggregation::SameValue);
+
+		AddColumn(ColumnRef);
+	}
+	//////////////////////////////////////////////////
+	// Inheritance Super Id Column (UStruct and UClass)
+	{
+		TSharedRef<FTableColumn> ColumnRef = MakeShared<FTableColumn>(FObjectTableColumns::InheritanceSuperIdColumnId);
+		FTableColumn& Column = *ColumnRef;
+
+		Column.SetIndex(ColumnIndex++);
+
+		Column.SetShortName(LOCTEXT("InheritanceSuperIdColumnName", "Inh. Super Id"));
+		Column.SetTitleName(LOCTEXT("InheritanceSuperIdColumnTitle", "Inheritance Super Id"));
+		Column.SetDescription(LOCTEXT("InheritanceSuperIdColumnDesc", "The id of the super class/struct (inheritance)"));
+
+		Column.SetFlags(ETableColumnFlags::CanBeHidden |
+						ETableColumnFlags::CanBeFiltered);
+
+		Column.SetHorizontalAlignment(HAlign_Left);
+		Column.SetInitialWidth(60.0f);
+
+		Column.SetDataType(ETableCellDataType::Int64);
+
+		class FValueGetter : public FAggregatedValueGetter
+		{
+		public:
+			virtual const TOptional<FTableCellValue> GetNodeValue(const FBaseTreeNode& Node) const override
+			{
+				if (Node.Is<FStructObjectNode>())
+				{
+					const FStructObjectNode& StructNode = Node.As<FStructObjectNode>();
+					const uint32 Id = StructNode.GetInheritanceSuperId();
+					if (Id != uint32(-1))
+					{
+						return FTableCellValue(int64(Id));
+					}
+				}
+				return TOptional<FTableCellValue>();
+			}
+		};
+		TSharedRef<ITableCellValueGetter> Getter = MakeShared<FValueGetter>();
+		Column.SetValueGetter(Getter);
+
+		TSharedRef<ITableCellValueFormatter> Formatter = MakeShared<FInt64ValueFormatterAsNumber>();
+		Column.SetValueFormatter(Formatter);
+
+		TSharedRef<ITableCellValueSorter> Sorter = MakeShared<FSorterByInt64Value>(ColumnRef);
+		Column.SetValueSorter(Sorter);
+		Column.SetInitialSortMode(EColumnSortMode::Descending);
+
+		//TSharedRef<IFilterValueConverter> Converter = MakeShared<FMemoryFilterValueConverter>();
+		//Column.SetValueConverter(Converter);
+
+		Column.SetAggregation(ETableColumnAggregation::SameValue);
+
+		AddColumn(ColumnRef);
+	}
+	//////////////////////////////////////////////////
 	// Structure Size Column
 	{
 		TSharedRef<FTableColumn> ColumnRef = MakeShared<FTableColumn>(FObjectTableColumns::StructureSizeColumnId);
@@ -757,10 +985,10 @@
 		public:
 			virtual const TOptional<FTableCellValue> GetNodeValue(const FBaseTreeNode& Node) const override
 			{
-				if (Node.Is<FObjectNode>())
+				if (Node.Is<FStructObjectNode>())
 				{
-					const FObjectNode& ObjectNode = Node.As<FObjectNode>();
-					return FTableCellValue(int64(ObjectNode.GetStructureSize()));
+					const FStructObjectNode& StructNode = Node.As<FStructObjectNode>();
+					return FTableCellValue(int64(StructNode.GetStructureSize()));
 				}
 				return TOptional<FTableCellValue>();
 			}
@@ -1072,59 +1300,86 @@
 		AddColumn(ColumnRef);
 	}
 	//////////////////////////////////////////////////
-	// Inclusive Estimated Size Column
+	// Impact Column
 	{
-		TSharedRef<FTableColumn> ColumnRef = MakeShared<FTableColumn>(FObjectTableColumns::InclEstimatedSizeColumnId);
+		TSharedRef<FTableColumn> ColumnRef = MakeShared<FTableColumn>(FObjectTableColumns::ImpactColumnId);
 		FTableColumn& Column = *ColumnRef;
 
 		Column.SetIndex(ColumnIndex++);
 
-		Column.SetShortName(LOCTEXT("InclEstimatedSizeColumnName", "Incl. Estimated Size"));
-		Column.SetTitleName(LOCTEXT("InclEstimatedSizeColumnTitle", "Inclusive Estimated Size (MiB)"));
-		Column.SetDescription(LOCTEXT("InclEstimatedSizeColumnDesc", "The total estimated memory size (including references), in MiB"));
+		Column.SetShortName(LOCTEXT("ImpactColumnName", "Impact"));
+		Column.SetTitleName(LOCTEXT("ImpactColumnTitle", "Impact (%)"));
+		Column.SetDescription(LOCTEXT("ImpactColumnDesc", "The impact (percentage) of estimated memory size"));
 
 		Column.SetFlags(ETableColumnFlags::CanBeHidden |
-						ETableColumnFlags::CanBeFiltered |
-						ETableColumnFlags::IsDynamic);
+			ETableColumnFlags::CanBeFiltered |
+			ETableColumnFlags::IsDynamic);
 
 		Column.SetHorizontalAlignment(HAlign_Right);
 		Column.SetInitialWidth(60.0f);
 
-		Column.SetDataType(ETableCellDataType::Int64);
+		Column.SetDataType(ETableCellDataType::Double);
 
-		class FValueGetter : public FAggregatedValueGetter
+		class FValueGetter : public FTableCellValueGetter
 		{
 		public:
-			virtual const TOptional<FTableCellValue> GetNodeValue(const FBaseTreeNode& Node) const override
+			virtual const TOptional<FTableCellValue> GetValue(const FTableColumn& Column, const FBaseTreeNode& Node) const override
 			{
-				if (Node.Is<FObjectNode>())
+				int64 EstimatedSize = 0;
+				if (Node.IsGroup() && Node.Is<FTableTreeNode>())
 				{
+					const FTableTreeNode& TableTreeNode = Node.As<FTableTreeNode>();
+					if (TableTreeNode.HasAggregatedValue(FObjectTableColumns::EstimatedSizeColumnId))
+					{
+						const FTableCellValue& Size = TableTreeNode.GetAggregatedValue(FObjectTableColumns::EstimatedSizeColumnId);
+						if (Size.DataType == ETableCellDataType::Int64 &&
+							Size.Int64 > 0)
+						{
+							EstimatedSize = Size.Int64;
+						}
+					}
+				}
+				else if (Node.Is<FObjectNode>())
+				{
 					const FObjectNode& ObjectNode = Node.As<FObjectNode>();
-					return FTableCellValue(static_cast<int64>(ObjectNode.GetInclusiveEstimatedSize()));
+					EstimatedSize = ObjectNode.GetEstimatedSize();
 				}
+
+				if (EstimatedSize > 0)
+				{
+					auto ParentNode = Node.GetParent();
+					if (ParentNode.IsValid() && ParentNode->IsGroup() && ParentNode->Is<FTableTreeNode>())
+					{
+						const FTableTreeNode& TableTreeNode = ParentNode->As<FTableTreeNode>();
+						if (TableTreeNode.HasAggregatedValue(FObjectTableColumns::EstimatedSizeColumnId))
+						{
+							const FTableCellValue& Total = TableTreeNode.GetAggregatedValue(FObjectTableColumns::EstimatedSizeColumnId);
+							if (Total.DataType == ETableCellDataType::Int64 &&
+								Total.Int64 > 0)
+							{
+								const double Ratio = double(EstimatedSize) / double(Total.Int64);
+								return FTableCellValue(static_cast<double>(Ratio));
+							}
+						}
+					}
+				}
+
 				return TOptional<FTableCellValue>();
 			}
 		};
 		TSharedRef<ITableCellValueGetter> Getter = MakeShared<FValueGetter>();
 		Column.SetValueGetter(Getter);
 
-		class FEstimatedSizeFormatter : public FTableCellValueFormatter
+		class FImpactFormatter : public FTableCellValueFormatter
 		{
 		public:
 			virtual FText FormatValue(const TOptional<FTableCellValue>& InValue) const override
 			{
 				if (InValue.IsSet())
 				{
-					const int64 Size = InValue.GetValue().Int64;
-					if (Size > 0)
-					{
-						int64 SizeMB = Size / (1024 * 1024);
-						if (SizeMB == 0)
-						{
-							return LOCTEXT("LessThan1MB", "< 1 MiB");
-						}
-						return FText::Format(LOCTEXT("EstimatedSizeFmt", "{0} MiB"), FText::AsNumber(SizeMB));
-					}
+					const double Ratio = InValue.GetValue().Double;
+					int64 Percent = FMath::RoundToInt(Ratio * 100.0);
+					return FText::Format(LOCTEXT("ImpactFmt", "{0}%"), FText::AsNumber(Percent));
 				}
 				return FText::GetEmpty();
 			}
@@ -1132,25 +1387,18 @@
 			{
 				if (InValue.IsSet())
 				{
-					const int64 Size = InValue.GetValue().Int64;
-					if (Size == 0)
-					{
-						return FText::GetEmpty(); // LOCTEXT("ZeroMB", "0");
-					}
+					const double Ratio = InValue.GetValue().Double;
 					FNumberFormattingOptions Options;
 					Options.MaximumFractionalDigits = 2;
-					return FText::Format(LOCTEXT("EstimatedSizeTooltipFmt", "{0} ({1} bytes)"),
-						FText::AsMemory(Size, &Options),
-						FText::AsNumber(Size));
+					return FText::Format(LOCTEXT("ImpactTooltipFmt", "{0}"), FText::AsPercent(Ratio, &Options));
 				}
 				return FText::GetEmpty();
 			}
 		};
-
-		TSharedRef<ITableCellValueFormatter> Formatter = MakeShared<FEstimatedSizeFormatter>();
+		TSharedRef<ITableCellValueFormatter> Formatter = MakeShared<FImpactFormatter>();
 		Column.SetValueFormatter(Formatter);
 
-		TSharedRef<ITableCellValueSorter> Sorter = MakeShared<FSorterByInt64Value>(ColumnRef);
+		TSharedRef<ITableCellValueSorter> Sorter = MakeShared<FSorterByDoubleValue>(ColumnRef);
 		Column.SetValueSorter(Sorter);
 		Column.SetInitialSortMode(EColumnSortMode::Descending);
 
@@ -1159,16 +1407,16 @@
 		AddColumn(ColumnRef);
 	}
 	//////////////////////////////////////////////////
-	// Impact Column
+	// Inclusive Estimated Size Column
 	{
-		TSharedRef<FTableColumn> ColumnRef = MakeShared<FTableColumn>(FObjectTableColumns::ImpactColumnId);
+		TSharedRef<FTableColumn> ColumnRef = MakeShared<FTableColumn>(FObjectTableColumns::InclEstimatedSizeColumnId);
 		FTableColumn& Column = *ColumnRef;
 
 		Column.SetIndex(ColumnIndex++);
 
-		Column.SetShortName(LOCTEXT("ImpactColumnName", "Impact"));
-		Column.SetTitleName(LOCTEXT("ImpactColumnTitle", "Impact (%)"));
-		Column.SetDescription(LOCTEXT("ImpactColumnDesc", "The impact (percentage) of estimated memory size"));
+		Column.SetShortName(LOCTEXT("InclEstimatedSizeColumnName", "Incl. Estimated Size"));
+		Column.SetTitleName(LOCTEXT("InclEstimatedSizeColumnTitle", "Inclusive Estimated Size (MiB)"));
+		Column.SetDescription(LOCTEXT("InclEstimatedSizeColumnDesc", "The total estimated memory size (including references), in MiB"));
 
 		Column.SetFlags(ETableColumnFlags::CanBeHidden |
 						ETableColumnFlags::CanBeFiltered |
@@ -1177,68 +1425,41 @@
 		Column.SetHorizontalAlignment(HAlign_Right);
 		Column.SetInitialWidth(60.0f);
 
-		Column.SetDataType(ETableCellDataType::Double);
+		Column.SetDataType(ETableCellDataType::Int64);
 
-		class FValueGetter : public FTableCellValueGetter
+		class FValueGetter : public FAggregatedValueGetter
 		{
 		public:
-			virtual const TOptional<FTableCellValue> GetValue(const FTableColumn& Column, const FBaseTreeNode& Node) const override
+			virtual const TOptional<FTableCellValue> GetNodeValue(const FBaseTreeNode& Node) const override
 			{
-				int64 EstimatedSize = 0;
-				if (Node.IsGroup() && Node.Is<FTableTreeNode>())
+				if (Node.Is<FObjectNode>())
 				{
-					const FTableTreeNode& TableTreeNode = Node.As<FTableTreeNode>();
-					if (TableTreeNode.HasAggregatedValue(FObjectTableColumns::EstimatedSizeColumnId))
-					{
-						const FTableCellValue& Size = TableTreeNode.GetAggregatedValue(FObjectTableColumns::EstimatedSizeColumnId);
-						if (Size.DataType == ETableCellDataType::Int64 &&
-							Size.Int64 > 0)
-						{
-							EstimatedSize = Size.Int64;
-						}
-					}
-				}
-				else if (Node.Is<FObjectNode>())
-				{
 					const FObjectNode& ObjectNode = Node.As<FObjectNode>();
-					EstimatedSize = ObjectNode.GetEstimatedSize();
+					return FTableCellValue(static_cast<int64>(ObjectNode.GetInclusiveEstimatedSize()));
 				}
-
-				if (EstimatedSize > 0)
-				{
-					auto ParentNode = Node.GetParent();
-					if (ParentNode.IsValid() && ParentNode->IsGroup() && ParentNode->Is<FTableTreeNode>())
-					{
-						const FTableTreeNode& TableTreeNode = ParentNode->As<FTableTreeNode>();
-						if (TableTreeNode.HasAggregatedValue(FObjectTableColumns::EstimatedSizeColumnId))
-						{
-							const FTableCellValue& Total = TableTreeNode.GetAggregatedValue(FObjectTableColumns::EstimatedSizeColumnId);
-							if (Total.DataType == ETableCellDataType::Int64 &&
-								Total.Int64 > 0)
-							{
-								const double Ratio = double(EstimatedSize) / double(Total.Int64);
-								return FTableCellValue(static_cast<double>(Ratio));
-							}
-						}
-					}
-				}
-
 				return TOptional<FTableCellValue>();
 			}
 		};
 		TSharedRef<ITableCellValueGetter> Getter = MakeShared<FValueGetter>();
 		Column.SetValueGetter(Getter);
 
-		class FImpactFormatter : public FTableCellValueFormatter
+		class FEstimatedSizeFormatter : public FTableCellValueFormatter
 		{
 		public:
 			virtual FText FormatValue(const TOptional<FTableCellValue>& InValue) const override
 			{
 				if (InValue.IsSet())
 				{
-					const double Ratio = InValue.GetValue().Double;
-					int64 Percent = FMath::RoundToInt(Ratio * 100.0);
-					return FText::Format(LOCTEXT("ImpactFmt", "{0}%"), FText::AsNumber(Percent));
+					const int64 Size = InValue.GetValue().Int64;
+					if (Size > 0)
+					{
+						int64 SizeMB = Size / (1024 * 1024);
+						if (SizeMB == 0)
+						{
+							return LOCTEXT("LessThan1MB", "< 1 MiB");
+						}
+						return FText::Format(LOCTEXT("EstimatedSizeFmt", "{0} MiB"), FText::AsNumber(SizeMB));
+					}
 				}
 				return FText::GetEmpty();
 			}
@@ -1246,18 +1467,25 @@
 			{
 				if (InValue.IsSet())
 				{
-					const double Ratio = InValue.GetValue().Double;
+					const int64 Size = InValue.GetValue().Int64;
+					if (Size == 0)
+					{
+						return FText::GetEmpty(); // LOCTEXT("ZeroMB", "0");
+					}
 					FNumberFormattingOptions Options;
 					Options.MaximumFractionalDigits = 2;
-					return FText::Format(LOCTEXT("ImpactTooltipFmt", "{0}"), FText::AsPercent(Ratio, &Options));
+					return FText::Format(LOCTEXT("EstimatedSizeTooltipFmt", "{0} ({1} bytes)"),
+						FText::AsMemory(Size, &Options),
+						FText::AsNumber(Size));
 				}
 				return FText::GetEmpty();
 			}
 		};
-		TSharedRef<ITableCellValueFormatter> Formatter = MakeShared<FImpactFormatter>();
+
+		TSharedRef<ITableCellValueFormatter> Formatter = MakeShared<FEstimatedSizeFormatter>();
 		Column.SetValueFormatter(Formatter);
 
-		TSharedRef<ITableCellValueSorter> Sorter = MakeShared<FSorterByDoubleValue>(ColumnRef);
+		TSharedRef<ITableCellValueSorter> Sorter = MakeShared<FSorterByInt64Value>(ColumnRef);
 		Column.SetValueSorter(Sorter);
 		Column.SetInitialSortMode(EColumnSortMode::Descending);
 
diff --git a/Engine/Source/Developer/TraceInsights/Private/Insights/ObjectProfiler/ViewModels/ObjectTable.h b/Engine/Source/Developer/TraceInsights/Private/Insights/ObjectProfiler/ViewModels/ObjectTable.h
--- a/Engine/Source/Developer/TraceInsights/Private/Insights/ObjectProfiler/ViewModels/ObjectTable.h
+++ b/Engine/Source/Developer/TraceInsights/Private/Insights/ObjectProfiler/ViewModels/ObjectTable.h
@@ -25,17 +25,27 @@ struct FObjectTableColumns
 	static const FName OuterIdColumnId;
 	static const FName PackageNameColumnId;
 	static const FName PackageIdColumnId;
+
+	static const FName PackageUniqueIdColumnId;
+	static const FName PackagePathColumnId;
+
 	static const FName ObjectFlagsColumnId;
+
 	static const FName ObjectPathColumnId;
 	static const FName VersePathColumnId;
+	
+	static const FName SuperIdColumnId;
+	static const FName InheritanceSuperIdColumnId;
 	static const FName StructureSizeColumnId;
+
 	static const FName SystemMemorySizeColumnId;
 	static const FName InclSystemMemorySizeColumnId;
 	static const FName VideoMemorySizeColumnId;
 	static const FName InclVideoMemorySizeColumnId;
 	static const FName EstimatedSizeColumnId;
-	static const FName InclEstimatedSizeColumnId;
 	static const FName ImpactColumnId;
+	static const FName InclEstimatedSizeColumnId;
+
 	static const FName ReferencesColumnId;
 };
 
diff --git a/Engine/Source/Developer/TraceInsights/Private/Insights/ObjectProfiler/ViewModels/ObjectTableViewPresets.cpp b/Engine/Source/Developer/TraceInsights/Private/Insights/ObjectProfiler/ViewModels/ObjectTableViewPresets.cpp
--- a/Engine/Source/Developer/TraceInsights/Private/Insights/ObjectProfiler/ViewModels/ObjectTableViewPresets.cpp
+++ b/Engine/Source/Developer/TraceInsights/Private/Insights/ObjectProfiler/ViewModels/ObjectTableViewPresets.cpp
@@ -55,6 +55,17 @@ TSharedRef<ITableTreeViewPreset> FObjectTableViewPresets::CreateAssetViewPreset(
 			{
 				InOutCurrentGroupings.Add(*GroupingByClass);
 			}
+
+			auto GroupingByObjectName = InAvailableGroupings.FindByPredicate(
+				[](TSharedPtr<FTreeNodeGrouping>& Grouping)
+				{
+					return Grouping->Is<FTreeNodeGroupingByPathBreakdown>() &&
+						Grouping->As<FTreeNodeGroupingByPathBreakdown>().GetColumnId() == FObjectTableColumns::ObjectNameColumnId;
+				});
+			if (GroupingByObjectName)
+			{
+				InOutCurrentGroupings.Add(*GroupingByObjectName);
+			}
 		}
 		virtual void GetColumnConfigSet(TArray<FTableColumnConfig>& InOutConfigSet) const override
 		{
@@ -259,22 +270,24 @@ TSharedRef<ITableTreeViewPreset> FObjectTableViewPresets::CreateOuterViewPreset(
 		}
 		virtual void GetColumnConfigSet(TArray<FTableColumnConfig>& InOutConfigSet) const override
 		{
-			InOutConfigSet.Add({ FTable::GetHierarchyColumnId(),                true, 400.0f });
-			InOutConfigSet.Add({ FObjectTableColumns::ObjectCountColumnId,      true, 60.0f });
-			InOutConfigSet.Add({ FObjectTableColumns::ObjectNameColumnId,       !true, 120.0f });
-			InOutConfigSet.Add({ FObjectTableColumns::ObjectIdColumnId,         !true, 60.0f });
-			InOutConfigSet.Add({ FObjectTableColumns::ClassNameColumnId,        true, 120.0f });
-			InOutConfigSet.Add({ FObjectTableColumns::ClassIdColumnId,          !true, 60.0f });
-			InOutConfigSet.Add({ FObjectTableColumns::OuterNameColumnId,        !true, 120.0f });
-			InOutConfigSet.Add({ FObjectTableColumns::OuterIdColumnId,          !true, 60.0f });
-			InOutConfigSet.Add({ FObjectTableColumns::PackageNameColumnId,      true, 120.0f });
-			InOutConfigSet.Add({ FObjectTableColumns::PackageIdColumnId,        !true, 60.0f });
-			InOutConfigSet.Add({ FObjectTableColumns::ObjectFlagsColumnId,      true, 140.0f });
-			InOutConfigSet.Add({ FObjectTableColumns::ObjectPathColumnId,       !true, 200.0f });
-			InOutConfigSet.Add({ FObjectTableColumns::VersePathColumnId,        !true, 200.0f });
-			InOutConfigSet.Add({ FObjectTableColumns::StructureSizeColumnId,    true, 100.0f });
-			InOutConfigSet.Add({ FObjectTableColumns::SystemMemorySizeColumnId, true, 100.0f });
-			InOutConfigSet.Add({ FObjectTableColumns::VideoMemorySizeColumnId,  true, 100.0f });
+			InOutConfigSet.Add({ FTable::GetHierarchyColumnId(),                    true, 400.0f });
+			InOutConfigSet.Add({ FObjectTableColumns::ObjectCountColumnId,          true, 60.0f });
+			InOutConfigSet.Add({ FObjectTableColumns::ObjectNameColumnId,           !true, 120.0f });
+			InOutConfigSet.Add({ FObjectTableColumns::ObjectIdColumnId,             !true, 60.0f });
+			InOutConfigSet.Add({ FObjectTableColumns::ClassNameColumnId,            true, 120.0f });
+			InOutConfigSet.Add({ FObjectTableColumns::ClassIdColumnId,              !true, 60.0f });
+			InOutConfigSet.Add({ FObjectTableColumns::OuterNameColumnId,            !true, 120.0f });
+			InOutConfigSet.Add({ FObjectTableColumns::OuterIdColumnId,              !true, 60.0f });
+			InOutConfigSet.Add({ FObjectTableColumns::PackageNameColumnId,          true, 120.0f });
+			InOutConfigSet.Add({ FObjectTableColumns::PackageIdColumnId,            !true, 60.0f });
+			InOutConfigSet.Add({ FObjectTableColumns::ObjectFlagsColumnId,          true, 140.0f });
+			InOutConfigSet.Add({ FObjectTableColumns::ObjectPathColumnId,           !true, 200.0f });
+			InOutConfigSet.Add({ FObjectTableColumns::VersePathColumnId,            !true, 200.0f });
+			InOutConfigSet.Add({ FObjectTableColumns::StructureSizeColumnId,        true, 100.0f });
+			InOutConfigSet.Add({ FObjectTableColumns::SystemMemorySizeColumnId,     !true, 100.0f });
+			InOutConfigSet.Add({ FObjectTableColumns::InclSystemMemorySizeColumnId, true, 100.0f });
+			InOutConfigSet.Add({ FObjectTableColumns::VideoMemorySizeColumnId,      !true, 100.0f });
+			InOutConfigSet.Add({ FObjectTableColumns::InclVideoMemorySizeColumnId,  true, 100.0f });
 		}
 		virtual void OnAppliedToView(STableTreeView& TableTreeView) const override
 		{
diff --git a/Engine/Source/Developer/TraceInsights/Private/Insights/ObjectProfiler/Widgets/SObjectTableTreeView.cpp b/Engine/Source/Developer/TraceInsights/Private/Insights/ObjectProfiler/Widgets/SObjectTableTreeView.cpp
--- a/Engine/Source/Developer/TraceInsights/Private/Insights/ObjectProfiler/Widgets/SObjectTableTreeView.cpp
+++ b/Engine/Source/Developer/TraceInsights/Private/Insights/ObjectProfiler/Widgets/SObjectTableTreeView.cpp
@@ -327,11 +327,6 @@ void SObjectTableTreeView::InternalCreateGroupings()
 				const FName ColumnId = Grouping->As<FTreeNodeGroupingByUniqueValue>().GetColumnId();
 				return (ColumnId == FObjectTableColumns::ObjectIdColumnId);
 			}
-			else if (Grouping->Is<FTreeNodeGroupingByPathBreakdown>())
-			{
-				const FName ColumnId = Grouping->As<FTreeNodeGroupingByPathBreakdown>().GetColumnId();
-				return (ColumnId == FObjectTableColumns::ObjectNameColumnId);
-			}
 			return false;
 		});
 
@@ -437,19 +432,37 @@ void SObjectTableTreeView::RebuildTree(bool bResync)
 						const TraceServices::FObjectInfo* Object = Snapshot->GetObject(ObjectId);
 						if (Object && Object->Id == ObjectId)
 						{
-							TSharedPtr<FObjectNode> ObjectNode = MakeShared<FObjectNode>(ObjectTable, *Object);
-							ObjectNodes[ObjectId] = ObjectNode;
-							++ValidObjectCount;
-
-							if (ObjectNode->IsClass())
+							TSharedPtr<FObjectNode> ObjectNode;
+							if (EnumHasAnyFlags(Object->FlagsEx, TraceServices::EObjectInfoFlags::IsPackage))
 							{
-								Classes.Add(ObjectNode);
+								auto PackageNode = MakeShared<FPackageObjectNode>(ObjectTable, *Object);
+								Packages.Add(PackageNode);
+								ObjectNode = PackageNode;
 							}
-
-							if (ObjectNode->IsPackage())
+							else if (EnumHasAnyFlags(Object->FlagsEx, TraceServices::EObjectInfoFlags::IsClass))
+							{
+								auto ClassNode = MakeShared<FClassObjectNode>(ObjectTable, *Object);
+								Classes.Add(ClassNode);
+								ObjectNode = ClassNode;
+							}
+							else if (EnumHasAnyFlags(Object->FlagsEx, TraceServices::EObjectInfoFlags::IsStruct))
+							{
+								UE_LOGF(LogObjectProfiler, Log, "[Obj] UStruct %u \"%ls\"", Object->Id, Object->Name);
+								auto StructNode = MakeShared<FStructObjectNode>(ObjectTable, *Object);
+								ObjectNode = StructNode;
+							}
+							else if (EnumHasAnyFlags(Object->FlagsEx, TraceServices::EObjectInfoFlags::IsField))
 							{
-								Packages.Add(ObjectNode);
+								UE_LOGF(LogObjectProfiler, Log, "[Obj] UField %u \"%ls\"", Object->Id, Object->Name);
+								auto FieldNode = MakeShared<FFieldObjectNode>(ObjectTable, *Object);
+								ObjectNode = FieldNode;
 							}
+							else
+							{
+								ObjectNode = MakeShared<FObjectNode>(ObjectTable, *Object);
+							}
+							ObjectNodes[ObjectId] = ObjectNode;
+							++ValidObjectCount;
 						}
 					}
 
diff --git a/Engine/Source/Developer/TraceServices/Private/Analyzers/ObjectTraceAnalysis.cpp b/Engine/Source/Developer/TraceServices/Private/Analyzers/ObjectTraceAnalysis.cpp
--- a/Engine/Source/Developer/TraceServices/Private/Analyzers/ObjectTraceAnalysis.cpp
+++ b/Engine/Source/Developer/TraceServices/Private/Analyzers/ObjectTraceAnalysis.cpp
@@ -264,7 +264,7 @@ bool FObjectAnalyzer::OnEvent(uint16 RouteId, EStyle Style, const FOnEventContex
 
 			if (StructureSize != 0)
 			{
-				Object->FlagsEx |= EObjectInfoFlags::IsStruct;
+				Object->FlagsEx |= EObjectInfoFlags::IsStruct | EObjectInfoFlags::IsField;
 			}
 		}
 		else
