Expose relation management functions to plugins.

Also export the RTTI data for FThreadTimingTrack so that plugins can use Is(), As() etc. functions on it.

#rnx

#Synced-CL 51619993

diff --git a/Engine/Source/Developer/TraceInsights/Private/Insights/Widgets/STimingView.h b/Engine/Source/Developer/TraceInsights/Private/Insights/Widgets/STimingView.h
--- a/Engine/Source/Developer/TraceInsights/Private/Insights/Widgets/STimingView.h
+++ b/Engine/Source/Developer/TraceInsights/Private/Insights/Widgets/STimingView.h
@@ -336,6 +336,11 @@ public:
 	virtual void PreventThrottling() override;
 	virtual void AddOverlayWidget(const TSharedRef<SWidget>& InWidget) override;
 
+	virtual const TArray<TUniquePtr<ITimingEventRelation>>& GetCurrentRelations() const override { return CurrentRelations; }
+	virtual TArray<TUniquePtr<ITimingEventRelation>>& EditCurrentRelations() override { return CurrentRelations; }
+	virtual void AddRelation(TUniquePtr<ITimingEventRelation>& Relation) override { CurrentRelations.Add(MoveTemp(Relation)); }
+	virtual void ClearRelations() override;
+
 	////////////////////////////////////////////////////////////////////////////////////////////////////
 
 	// The callback should return 'true' to continue the enumeration.
@@ -430,11 +435,6 @@ public:
 
 	const TSharedPtr<FBaseTimingTrack> GetTrackAt(float InPosX, float InPosY) const;
 
-	const TArray<TUniquePtr<ITimingEventRelation>>& GetCurrentRelations() const { return CurrentRelations; }
-	TArray<TUniquePtr<ITimingEventRelation>>& EditCurrentRelations() { return CurrentRelations; }
-	void AddRelation(TUniquePtr<ITimingEventRelation>& Relation) { CurrentRelations.Add(MoveTemp(Relation)); }
-	void ClearRelations();
-
 	TSharedPtr<FUICommandList> GetCommandList() { return CommandList; }
 
 	void CloseQuickFindTab();
diff --git a/Engine/Source/Developer/TraceInsights/Public/Insights/ITimingViewSession.h b/Engine/Source/Developer/TraceInsights/Public/Insights/ITimingViewSession.h
--- a/Engine/Source/Developer/TraceInsights/Public/Insights/ITimingViewSession.h
+++ b/Engine/Source/Developer/TraceInsights/Public/Insights/ITimingViewSession.h
@@ -17,6 +17,7 @@ class SWidget;
 class FBaseTimingTrack;
 class FTimingEventsTrack;
 class ITimingEvent;
+class ITimingEventRelation;
 enum class ETimingTrackLocation : uint32; // Insights/ViewModels/BaseTimingTrack.h
 
 #if UE_INSIGHTS_BACKWARD_COMPATIBILITY_UE54
@@ -181,6 +182,20 @@ public:
 
 	/** Adds a slot to the overlay. */
 	virtual void AddOverlayWidget(const TSharedRef<SWidget>& InWidget) = 0;
+
+	//////////////////////////////////////////////////
+
+	/** Get the active relations for this view. */
+	virtual const TArray<TUniquePtr<ITimingEventRelation>>& GetCurrentRelations() const = 0;
+
+	/** Get an editable array of relations for this view. */
+	virtual TArray<TUniquePtr<ITimingEventRelation>>& EditCurrentRelations() = 0;
+
+	/** Add a new relation to this view. */
+	virtual void AddRelation(TUniquePtr<ITimingEventRelation>& Relation) = 0;
+
+	/** Remove all relations from this view. */
+	virtual void ClearRelations() = 0;
 };
 
 ////////////////////////////////////////////////////////////////////////////////////////////////////
diff --git a/Engine/Source/Developer/TraceInsights/Public/Insights/TimingProfiler/Tracks/ThreadTimingTrack.h b/Engine/Source/Developer/TraceInsights/Public/Insights/TimingProfiler/Tracks/ThreadTimingTrack.h
--- a/Engine/Source/Developer/TraceInsights/Public/Insights/TimingProfiler/Tracks/ThreadTimingTrack.h
+++ b/Engine/Source/Developer/TraceInsights/Public/Insights/TimingProfiler/Tracks/ThreadTimingTrack.h
@@ -12,7 +12,7 @@ namespace UE::Insights::TimingProfiler
 
 class FThreadTimingTrack : public FTimingEventsTrack
 {
-	INSIGHTS_DECLARE_RTTI(FThreadTimingTrack, FTimingEventsTrack)
+	INSIGHTS_DECLARE_RTTI(FThreadTimingTrack, FTimingEventsTrack, TRACEINSIGHTS_API)
 
 public:
 	explicit FThreadTimingTrack() {}
