[Insights] TimingView: Fixed the horizontal scrollbar to have its own reseverd area (like the vertical scrollbar) instead of overlapping the tracks area.

#Synced-CL 51663925

diff --git a/Engine/Source/Developer/TraceInsights/Private/Insights/Widgets/STimingView.cpp b/Engine/Source/Developer/TraceInsights/Private/Insights/Widgets/STimingView.cpp
--- a/Engine/Source/Developer/TraceInsights/Private/Insights/Widgets/STimingView.cpp
+++ b/Engine/Source/Developer/TraceInsights/Private/Insights/Widgets/STimingView.cpp
@@ -320,7 +320,7 @@ void STimingView::Construct(const FArguments& InArgs, FName InViewName)
 
 		+ SOverlay::Slot()
 		.VAlign(VAlign_Bottom)
-		.Padding(FMargin(0.0f, 0.0f, 8.0f, 0.0f))
+		.Padding(FMargin(0.0f, 0.0f, 10.0f, 0.0f))
 		[
 			SAssignNew(HorizontalScrollBar, SScrollBar)
 			.Orientation(Orient_Horizontal)
@@ -739,8 +739,11 @@ void STimingView::Tick(const FGeometry& AllottedGeometry, const double InCurrent
 	////////////////////////////////////////////////////////////////////////////////////////////////////
 	// Update viewport.
 
-	Viewport.SetPosY(32.0f); // height of toolbar
-	Viewport.UpdateSize(FMath::RoundToFloat(ViewWidth), FMath::RoundToFloat(ViewHeight) - Viewport.GetPosY());
+	constexpr float ToolbarHeight = 32.0f;
+	constexpr float HorizontalScrollBarHeight = 12.0f;
+
+	Viewport.SetPosY(ToolbarHeight);
+	Viewport.UpdateSize(FMath::RoundToFloat(ViewWidth), FMath::RoundToFloat(ViewHeight) - Viewport.GetPosY() - HorizontalScrollBarHeight);
 
 	if (!bIsPanning && !bAutoScroll)
 	{
@@ -1048,7 +1051,7 @@ PRAGMA_ENABLE_DEPRECATION_WARNINGS
 	}
 
 	// Set the VerticalScrollBar padding so it is limited to the scrollable area.
-	VerticalScrollBar->SetPadding(FMargin(0.0f, Viewport.GetPosY() + TopOffset + 2.0f, 0.0f, FMath::Max(BottomOffset + 2.0f, 10.0f)));
+	VerticalScrollBar->SetPadding(FMargin(0.0f, Viewport.GetPosY() + TopOffset + 2.0f, 0.0f, BottomOffset + HorizontalScrollBarHeight));
 
 	////////////////////////////////////////////////////////////////////////////////////////////////////
 
@@ -1726,15 +1729,29 @@ int32 STimingView::OnPaint(const FPaintArgs& Args, const FGeometry& AllottedGeom
 	// Draw tooltip with info about hovered event.
 	Tooltip.Draw(DrawContext);
 
+	const FLinearColor BackgroundColor = FSlateColor(EStyleColor::Panel).GetSpecifiedColor();
+
 	// Fill the background of the toolbar.
-	DrawContext.DrawBox(0.0f, 0.0f, ViewWidth, Viewport.GetPosY(), WhiteBrush, FSlateColor(EStyleColor::Panel).GetSpecifiedColor());
+	DrawContext.DrawBox(0.0f, 0.0f, ViewWidth, Viewport.GetPosY(), WhiteBrush, BackgroundColor);
 
 	// Fill the background of the vertical scrollbar.
-	const float ScrollBarHeight = Viewport.GetScrollableAreaHeight();
-	if (ScrollBarHeight > 0)
 	{
-		constexpr float ScrollBarWidth = 12.0f;
-		DrawContext.DrawBox(ViewWidth - ScrollBarWidth, Viewport.GetPosY() + Viewport.GetTopOffset(), ScrollBarWidth, ScrollBarHeight, WhiteBrush, FSlateColor(EStyleColor::Panel).GetSpecifiedColor());
+		const float ScrollBarHeight = Viewport.GetScrollableAreaHeight();
+		if (ScrollBarHeight > 0)
+		{
+			constexpr float ScrollBarWidth = 12.0f;
+			const     float ScrollBarPosX = ViewWidth - ScrollBarWidth;
+			const     float ScrollBarPosY = Viewport.GetPosY() + Viewport.GetTopOffset();
+			DrawContext.DrawBox(ScrollBarPosX, ScrollBarPosY, ScrollBarWidth, ScrollBarHeight, WhiteBrush, BackgroundColor);
+		}
+	}
+	// Fill the background of the horizontal scrollbar.
+	{
+		const     float ScrollBarWidth = ViewWidth;
+		constexpr float ScrollBarHeight = 12.0f;
+		constexpr float ScrollBarPosX = 0.0f;
+		const     float ScrollBarPosY = ViewHeight - ScrollBarHeight;
+		DrawContext.DrawBox(ScrollBarPosX, ScrollBarPosY, ScrollBarWidth, ScrollBarHeight, WhiteBrush, BackgroundColor);
 	}
 
 	//////////////////////////////////////////////////
