[Insights]
* ObjectTrace: Added HideObjectsWithLowInclusiveEstimatedMemory filtering. Only shows objects with Inclusive Estimated Memory (System + Video) >= 1 MiB.
* ObjectTrace: Added HideObjectsWithLowEstimatedMemory filtering. Only shows objects with Exclusive Estimated Memory (System + Video) >= 1 MiB.
* ObjectTrace: Added HideObjectsWithLowImpact filtering. Only shows objects with Impact % >= 1%.
* ObjectTrace: Added custom FFunctionObjectNode for UFunction objects.

#rnx

#Synced-CL 51883135

diff --git a/Engine/Source/Developer/TraceInsights/Private/Insights/ObjectProfiler/ViewModels/ObjectGroupingByOuter.cpp b/Engine/Source/Developer/TraceInsights/Private/Insights/ObjectProfiler/ViewModels/ObjectGroupingByOuter.cpp
--- a/Engine/Source/Developer/TraceInsights/Private/Insights/ObjectProfiler/ViewModels/ObjectGroupingByOuter.cpp
+++ b/Engine/Source/Developer/TraceInsights/Private/Insights/ObjectProfiler/ViewModels/ObjectGroupingByOuter.cpp
@@ -76,29 +76,6 @@ void FObjectGroupingByOuter::GroupNodes(
 
 ////////////////////////////////////////////////////////////////////////////////////////////////////
 
-static TSharedRef<FObjectNode> CloneObjectNode(FObjectNode& InNode)
-{
-	if (InNode.Is<FPackageObjectNode>())
-	{
-		return MakeShared<FPackageObjectNode>(InNode.As<FPackageObjectNode>());
-	}
-	if (InNode.Is<FClassObjectNode>())
-	{
-		return MakeShared<FClassObjectNode>(InNode.As<FClassObjectNode>());
-	}
-	if (InNode.Is<FStructObjectNode>())
-	{
-		return MakeShared<FStructObjectNode>(InNode.As<FStructObjectNode>());
-	}
-	if (InNode.Is<FFieldObjectNode>())
-	{
-		return MakeShared<FFieldObjectNode>(InNode.As<FFieldObjectNode>());
-	}
-	return MakeShared<FObjectNode>(InNode);
-}
-
-////////////////////////////////////////////////////////////////////////////////////////////////////
-
 TSharedPtr<FObjectNode> FObjectGroupingByOuter::AddObject(FObjectNode& ObjNode) const
 {
 	const uint32 ObjId = ObjNode.GetObjectId();
@@ -111,7 +88,7 @@ TSharedPtr<FObjectNode> FObjectGroupingByOuter::AddObject(FObjectNode& ObjNode)
 		return ObjectNodes[ObjId];
 	}
 
-	TSharedPtr<FObjectNode> NewNode = CloneObjectNode(ObjNode);
+	TSharedPtr<FObjectNode> NewNode = FObjectNode::Clone(ObjNode);
 	NewNode->SetIsMissing();
 	ObjectNodes[ObjId] = NewNode;
 
diff --git a/Engine/Source/Developer/TraceInsights/Private/Insights/ObjectProfiler/ViewModels/ObjectNode.cpp b/Engine/Source/Developer/TraceInsights/Private/Insights/ObjectProfiler/ViewModels/ObjectNode.cpp
--- a/Engine/Source/Developer/TraceInsights/Private/Insights/ObjectProfiler/ViewModels/ObjectNode.cpp
+++ b/Engine/Source/Developer/TraceInsights/Private/Insights/ObjectProfiler/ViewModels/ObjectNode.cpp
@@ -16,6 +16,7 @@ INSIGHTS_IMPLEMENT_RTTI(FObjectNode)
 INSIGHTS_IMPLEMENT_RTTI(FFieldObjectNode)
 INSIGHTS_IMPLEMENT_RTTI(FStructObjectNode)
 INSIGHTS_IMPLEMENT_RTTI(FClassObjectNode)
+INSIGHTS_IMPLEMENT_RTTI(FFunctionObjectNode)
 INSIGHTS_IMPLEMENT_RTTI(FPackageObjectNode)
 
 ////////////////////////////////////////////////////////////////////////////////////////////////////
@@ -137,6 +138,80 @@ TSharedPtr<FObjectNode> FObjectNode::GetPackage() const
 	return nullptr;
 }
 
+////////////////////////////////////////////////////////////////////////////////////////////////////
+
+// static
+int64 FObjectNode::GetEstimatedMemorySize(const FBaseTreeNode& Node)
+{
+	if (Node.Is<FObjectNode>())
+	{
+		return Node.As<FObjectNode>().GetEstimatedMemorySize();
+	}
+
+	if (Node.IsGroup() && Node.Is<FTableTreeNode>())
+	{
+		const FTableTreeNode& TableNode = Node.As<FTableTreeNode>();
+		if (TableNode.HasAggregatedValue(FObjectTableColumns::EstimatedSizeColumnId))
+		{
+			const FTableCellValue& Size = TableNode.GetAggregatedValue(FObjectTableColumns::EstimatedSizeColumnId);
+			if (Size.DataType == ETableCellDataType::Int64)
+			{
+				return Size.Int64;
+			}
+		}
+	}
+
+	return 0;
+}
+
+////////////////////////////////////////////////////////////////////////////////////////////////////
+
+double FObjectNode::GetEstimatedMemoryImpact() const
+{
+	const int64 EstimatedMemSize = GetEstimatedMemorySize();
+	if (EstimatedMemSize > 0)
+	{
+		TSharedPtr<FBaseTreeNode> ParentNode = GetParent();
+		if (ParentNode.IsValid())
+		{
+			const int64 ParentEstimatedMemSize = FObjectNode::GetEstimatedMemorySize(*ParentNode);
+			if (ParentEstimatedMemSize > 0)
+			{
+				return double(EstimatedMemSize) / double(ParentEstimatedMemSize);
+			}
+		}
+	}
+	return 0.0;
+}
+
+////////////////////////////////////////////////////////////////////////////////////////////////////
+
+//static
+TSharedRef<FObjectNode> FObjectNode::Clone(FObjectNode& InNode)
+{
+	if (InNode.Is<FPackageObjectNode>())
+	{
+		return MakeShared<FPackageObjectNode>(InNode.As<FPackageObjectNode>());
+	}
+	if (InNode.Is<FFunctionObjectNode>())
+	{
+		return MakeShared<FFunctionObjectNode>(InNode.As<FFunctionObjectNode>());
+	}
+	if (InNode.Is<FClassObjectNode>())
+	{
+		return MakeShared<FClassObjectNode>(InNode.As<FClassObjectNode>());
+	}
+	if (InNode.Is<FStructObjectNode>())
+	{
+		return MakeShared<FStructObjectNode>(InNode.As<FStructObjectNode>());
+	}
+	if (InNode.Is<FFieldObjectNode>())
+	{
+		return MakeShared<FFieldObjectNode>(InNode.As<FFieldObjectNode>());
+	}
+	return MakeShared<FObjectNode>(InNode);
+}
+
 ////////////////////////////////////////////////////////////////////////////////////////////////////
 // FFieldObjectNode
 ////////////////////////////////////////////////////////////////////////////////////////////////////
@@ -236,6 +311,43 @@ FLinearColor FClassObjectNode::GetIconColor() const
 	return FLinearColor(0.5f, 1.0f, 1.0f, 1.0f);
 }
 
+////////////////////////////////////////////////////////////////////////////////////////////////////
+// FFunctionObjectNode
+////////////////////////////////////////////////////////////////////////////////////////////////////
+
+FFunctionObjectNode::FFunctionObjectNode(TWeakPtr<FObjectTable> InParentTable, const TraceServices::FObjectInfo& ObjectInfo)
+	: FStructObjectNode(InParentTable, ObjectInfo)
+	, FunctionFlags(ObjectInfo.FunctionFlags)
+	, FunctionNumParms(ObjectInfo.FunctionNumParms)
+	, FunctionParmsSize(ObjectInfo.FunctionParmsSize)
+{
+	ensure(EnumHasAnyFlags(ObjectInfo.FlagsEx, TraceServices::EObjectInfoFlags::IsFunction));
+}
+
+////////////////////////////////////////////////////////////////////////////////////////////////////
+
+FFunctionObjectNode::FFunctionObjectNode(const FFunctionObjectNode& Clone)
+	: FStructObjectNode(Clone)
+	, FunctionFlags(Clone.FunctionFlags)
+	, FunctionNumParms(Clone.FunctionNumParms)
+	, FunctionParmsSize(Clone.FunctionParmsSize)
+{
+}
+
+////////////////////////////////////////////////////////////////////////////////////////////////////
+
+const FSlateBrush* FFunctionObjectNode::GetIcon() const
+{
+	return FInsightsStyle::GetBrush("Icons.UFunction.TreeItem");
+}
+
+////////////////////////////////////////////////////////////////////////////////////////////////////
+
+FLinearColor FFunctionObjectNode::GetIconColor() const
+{
+	return FLinearColor(0.7f, 0.7f, 1.0f, 1.0f);
+}
+
 ////////////////////////////////////////////////////////////////////////////////////////////////////
 // FPackageObjectNode
 ////////////////////////////////////////////////////////////////////////////////////////////////////
diff --git a/Engine/Source/Developer/TraceInsights/Private/Insights/ObjectProfiler/ViewModels/ObjectNode.h b/Engine/Source/Developer/TraceInsights/Private/Insights/ObjectProfiler/ViewModels/ObjectNode.h
--- a/Engine/Source/Developer/TraceInsights/Private/Insights/ObjectProfiler/ViewModels/ObjectNode.h
+++ b/Engine/Source/Developer/TraceInsights/Private/Insights/ObjectProfiler/ViewModels/ObjectNode.h
@@ -80,11 +80,14 @@ public:
 
 	int64 GetSystemMemorySize() const { return SystemMemorySize; }
 	int64 GetVideoMemorySize() const { return VideoMemorySize; }
-	int64 GetEstimatedSize() const { return SystemMemorySize + VideoMemorySize; }
+	int64 GetEstimatedMemorySize() const { return SystemMemorySize + VideoMemorySize; }
+
+	static int64 GetEstimatedMemorySize(const FBaseTreeNode& Node);
+	double GetEstimatedMemoryImpact() const;
 
 	int64 GetInclusiveSystemMemorySize() const { return InclSystemMemorySize; }
 	int64 GetInclusiveVideoMemorySize() const { return InclVideoMemorySize; }
-	int64 GetInclusiveEstimatedSize() const { return InclSystemMemorySize + InclVideoMemorySize; }
+	int64 GetInclusiveEstimatedMemorySize() const { return InclSystemMemorySize + InclVideoMemorySize; }
 
 	virtual uint32 GetNumReferences() const { return 0; }
 
@@ -93,8 +96,11 @@ public:
 	virtual bool IsField() const { return false; }
 	virtual bool IsStruct() const { return false; }
 	virtual bool IsClass() const { return false; }
+	virtual bool IsFunction() const { return false; }
 	virtual bool IsPackage() const { return false; }
 
+	static TSharedRef<FObjectNode> Clone(FObjectNode& InNode);
+
 protected:
 	void SetIsMissing() { InternalFlags = EInternalFlags(uint32(InternalFlags) | uint32(EInternalFlags::IsMissing)); }
 	void ResetIsMissing() { InternalFlags = EInternalFlags(uint32(InternalFlags) & ~uint32(EInternalFlags::IsMissing)); }
@@ -196,6 +202,35 @@ public:
 
 ////////////////////////////////////////////////////////////////////////////////////////////////////
 
+class FFunctionObjectNode : public FStructObjectNode
+{
+	friend class SObjectTableTreeView;
+
+	INSIGHTS_DECLARE_RTTI(FFunctionObjectNode, FStructObjectNode)
+
+	virtual const FSlateBrush* GetIcon() const override;
+	virtual FLinearColor GetIconColor() const override;
+
+	virtual bool IsFunction() const override { return true; }
+
+	uint32 GetFunctionFlags() const { return FunctionFlags; }
+	uint32 GetFunctionNumParms() const { return uint32(FunctionNumParms); }
+	uint32 GetFunctionParmsSize() const { return uint32(FunctionParmsSize); }
+
+public:
+	/** Initialization constructor for the UFunction node. */
+	explicit FFunctionObjectNode(TWeakPtr<FObjectTable> InParentTable, const TraceServices::FObjectInfo& ObjectInfo);
+
+	explicit FFunctionObjectNode(const FFunctionObjectNode& Clone);
+
+private:
+	uint32 FunctionFlags = 0;
+	uint8 FunctionNumParms = 0;
+	uint16 FunctionParmsSize = 0;
+};
+
+////////////////////////////////////////////////////////////////////////////////////////////////////
+
 class FPackageObjectNode : public FObjectNode
 {
 	friend class SObjectTableTreeView;
diff --git a/Engine/Source/Developer/TraceInsights/Private/Insights/ObjectProfiler/ViewModels/ObjectTable.cpp b/Engine/Source/Developer/TraceInsights/Private/Insights/ObjectProfiler/ViewModels/ObjectTable.cpp
--- a/Engine/Source/Developer/TraceInsights/Private/Insights/ObjectProfiler/ViewModels/ObjectTable.cpp
+++ b/Engine/Source/Developer/TraceInsights/Private/Insights/ObjectProfiler/ViewModels/ObjectTable.cpp
@@ -913,7 +913,7 @@ void FObjectTable::AddDefaultColumns()
 
 		Column.SetIndex(ColumnIndex++);
 
-		Column.SetShortName(LOCTEXT("InheritanceSuperIdColumnName", "Inh. Super Id"));
+		Column.SetShortName(LOCTEXT("InheritanceSuperIdColumnName", "Inheritance"));
 		Column.SetTitleName(LOCTEXT("InheritanceSuperIdColumnTitle", "Inheritance Super Id"));
 		Column.SetDescription(LOCTEXT("InheritanceSuperIdColumnDesc", "The id of the super class/struct (inheritance)"));
 
@@ -1241,7 +1241,7 @@ void FObjectTable::AddDefaultColumns()
 				if (Node.Is<FObjectNode>())
 				{
 					const FObjectNode& ObjectNode = Node.As<FObjectNode>();
-					return FTableCellValue(static_cast<int64>(ObjectNode.GetEstimatedSize()));
+					return FTableCellValue(static_cast<int64>(ObjectNode.GetEstimatedMemorySize()));
 				}
 				return TOptional<FTableCellValue>();
 			}
@@ -1325,45 +1325,20 @@ void FObjectTable::AddDefaultColumns()
 		public:
 			virtual const TOptional<FTableCellValue> GetValue(const FTableColumn& Column, const FBaseTreeNode& Node) const override
 			{
-				int64 EstimatedSize = 0;
-				if (Node.IsGroup() && Node.Is<FTableTreeNode>())
-				{
-					const FTableTreeNode& TableTreeNode = Node.As<FTableTreeNode>();
-					if (TableTreeNode.HasAggregatedValue(FObjectTableColumns::EstimatedSizeColumnId))
-					{
-						const FTableCellValue& Size = TableTreeNode.GetAggregatedValue(FObjectTableColumns::EstimatedSizeColumnId);
-						if (Size.DataType == ETableCellDataType::Int64 &&
-							Size.Int64 > 0)
-						{
-							EstimatedSize = Size.Int64;
-						}
-					}
-				}
-				else if (Node.Is<FObjectNode>())
-				{
-					const FObjectNode& ObjectNode = Node.As<FObjectNode>();
-					EstimatedSize = ObjectNode.GetEstimatedSize();
-				}
-
+				const int64 EstimatedSize = FObjectNode::GetEstimatedMemorySize(Node);
 				if (EstimatedSize > 0)
 				{
-					auto ParentNode = Node.GetParent();
-					if (ParentNode.IsValid() && ParentNode->IsGroup() && ParentNode->Is<FTableTreeNode>())
+					TSharedPtr<FBaseTreeNode> ParentNode = Node.GetParent();
+					if (ParentNode.IsValid())
 					{
-						const FTableTreeNode& TableTreeNode = ParentNode->As<FTableTreeNode>();
-						if (TableTreeNode.HasAggregatedValue(FObjectTableColumns::EstimatedSizeColumnId))
+						const int64 ParentEstimatedMemSize = FObjectNode::GetEstimatedMemorySize(*ParentNode);
+						if (ParentEstimatedMemSize > 0)
 						{
-							const FTableCellValue& Total = TableTreeNode.GetAggregatedValue(FObjectTableColumns::EstimatedSizeColumnId);
-							if (Total.DataType == ETableCellDataType::Int64 &&
-								Total.Int64 > 0)
-							{
-								const double Ratio = double(EstimatedSize) / double(Total.Int64);
-								return FTableCellValue(static_cast<double>(Ratio));
-							}
+							const double Ratio = double(EstimatedSize) / double(ParentEstimatedMemSize);
+							return FTableCellValue(static_cast<double>(Ratio));
 						}
 					}
 				}
-
 				return TOptional<FTableCellValue>();
 			}
 		};
@@ -1435,7 +1410,7 @@ void FObjectTable::AddDefaultColumns()
 				if (Node.Is<FObjectNode>())
 				{
 					const FObjectNode& ObjectNode = Node.As<FObjectNode>();
-					return FTableCellValue(static_cast<int64>(ObjectNode.GetInclusiveEstimatedSize()));
+					return FTableCellValue(static_cast<int64>(ObjectNode.GetInclusiveEstimatedMemorySize()));
 				}
 				return TOptional<FTableCellValue>();
 			}
diff --git a/Engine/Source/Developer/TraceInsights/Private/Insights/ObjectProfiler/Widgets/SObjectTableTreeView.cpp b/Engine/Source/Developer/TraceInsights/Private/Insights/ObjectProfiler/Widgets/SObjectTableTreeView.cpp
--- a/Engine/Source/Developer/TraceInsights/Private/Insights/ObjectProfiler/Widgets/SObjectTableTreeView.cpp
+++ b/Engine/Source/Developer/TraceInsights/Private/Insights/ObjectProfiler/Widgets/SObjectTableTreeView.cpp
@@ -424,7 +424,13 @@ void SObjectTableTreeView::RebuildTree(bool bResync)
 
 					const uint32 ObjectCount = Snapshot->GetNumObjects();
 					ObjectNodes.AddDefaulted(ObjectCount); // index == object id
+
 					uint32 ValidObjectCount = 0;
+					uint32 PackageCount = 0;
+					uint32 FunctionCount = 0;
+					uint32 ClassCount = 0;
+					uint32 StructCount = 0;
+					uint32 FieldCount = 0;
 
 					TSharedPtr<FObjectTable> ObjectTable = StaticCastSharedPtr<FObjectTable>(GetTable());
 					for (uint32 ObjectId = 0; ObjectId < ObjectCount; ++ObjectId)
@@ -438,24 +444,32 @@ void SObjectTableTreeView::RebuildTree(bool bResync)
 								auto PackageNode = MakeShared<FPackageObjectNode>(ObjectTable, *Object);
 								Packages.Add(PackageNode);
 								ObjectNode = PackageNode;
+								++PackageCount;
+							}
+							else if (EnumHasAnyFlags(Object->FlagsEx, TraceServices::EObjectInfoFlags::IsFunction))
+							{
+								auto FunctionNode = MakeShared<FFunctionObjectNode>(ObjectTable, *Object);
+								ObjectNode = FunctionNode;
+								++FunctionCount;
 							}
 							else if (EnumHasAnyFlags(Object->FlagsEx, TraceServices::EObjectInfoFlags::IsClass))
 							{
 								auto ClassNode = MakeShared<FClassObjectNode>(ObjectTable, *Object);
 								Classes.Add(ClassNode);
 								ObjectNode = ClassNode;
+								++ClassCount;
 							}
 							else if (EnumHasAnyFlags(Object->FlagsEx, TraceServices::EObjectInfoFlags::IsStruct))
 							{
-								UE_LOGF(LogObjectProfiler, Log, "[Obj] UStruct %u \"%ls\"", Object->Id, Object->Name);
 								auto StructNode = MakeShared<FStructObjectNode>(ObjectTable, *Object);
 								ObjectNode = StructNode;
+								++StructCount;
 							}
 							else if (EnumHasAnyFlags(Object->FlagsEx, TraceServices::EObjectInfoFlags::IsField))
 							{
-								UE_LOGF(LogObjectProfiler, Log, "[Obj] UField %u \"%ls\"", Object->Id, Object->Name);
 								auto FieldNode = MakeShared<FFieldObjectNode>(ObjectTable, *Object);
 								ObjectNode = FieldNode;
+								++FieldCount;
 							}
 							else
 							{
@@ -466,6 +480,9 @@ void SObjectTableTreeView::RebuildTree(bool bResync)
 						}
 					}
 
+					UE_LOGF(LogObjectProfiler, Log, "[Obj] Created nodes for %u objects (%u packages, %u functions, %u classes, %u structs, %u fields)",
+						ValidObjectCount, PackageCount, FunctionCount, ClassCount, StructCount, FieldCount);
+
 					TableRowNodes.Reserve(ValidObjectCount);
 
 					for (TSharedPtr<FObjectNode>& NodePtr : ObjectNodes)
@@ -560,7 +577,11 @@ void SObjectTableTreeView::ExtendMenuAfterMisc(FMenuBuilder& MenuBuilder)
 
 bool SObjectTableTreeView::HasCustomNodeFilter() const
 {
-	return bHideObjectsWithZeroStructureSize || bHideObjectsWithZeroMemorySize;
+	return bHideObjectsWithZeroStructureSize
+		|| bHideObjectsWithZeroEstimatedMemory
+		|| bHideObjectsWithLowInclusiveEstimatedMemory
+		|| bHideObjectsWithLowEstimatedMemory
+		|| bHideObjectsWithLowImpact;
 }
 
 ////////////////////////////////////////////////////////////////////////////////////////////////////
@@ -577,9 +598,33 @@ bool SObjectTableTreeView::FilterNodeCustom(const FTableTreeNode& InNode) const
 				return false;
 			}
 		}
-		if (bHideObjectsWithZeroMemorySize)
+		if (bHideObjectsWithZeroEstimatedMemory)
+		{
+			if (ObjectNode.GetEstimatedMemorySize() == 0)
+			{
+				return false;
+			}
+		}
+		if (bHideObjectsWithLowInclusiveEstimatedMemory)
+		{
+			const int64 EstimatedMemSize = ObjectNode.GetInclusiveEstimatedMemorySize();
+			if (EstimatedMemSize < 1024 * 1024)
+			{
+				return false;
+			}
+		}
+		if (bHideObjectsWithLowEstimatedMemory)
+		{
+			const int64 EstimatedMemSize = ObjectNode.GetEstimatedMemorySize();
+			if (EstimatedMemSize < 1024 * 1024)
+			{
+				return false;
+			}
+		}
+		if (bHideObjectsWithLowImpact)
 		{
-			if (ObjectNode.GetSystemMemorySize() == 0 && ObjectNode.GetVideoMemorySize() == 0)
+			const double EstimatedImpact = ObjectNode.GetEstimatedMemoryImpact();
+			if (EstimatedImpact < 0.01)
 			{
 				return false;
 			}
diff --git a/Engine/Source/Developer/TraceInsights/Private/Insights/ObjectProfiler/Widgets/SObjectTableTreeView.h b/Engine/Source/Developer/TraceInsights/Private/Insights/ObjectProfiler/Widgets/SObjectTableTreeView.h
--- a/Engine/Source/Developer/TraceInsights/Private/Insights/ObjectProfiler/Widgets/SObjectTableTreeView.h
+++ b/Engine/Source/Developer/TraceInsights/Private/Insights/ObjectProfiler/Widgets/SObjectTableTreeView.h
@@ -95,8 +95,11 @@ private:
 	//////////////////////////////////////////////////
 	// Filters
 
-	bool bHideObjectsWithZeroStructureSize = false;
-	bool bHideObjectsWithZeroMemorySize = false;
+	bool bHideObjectsWithZeroStructureSize = false; // only show UStruct objects (including UClass and UFunction objects; with valid Structure size)
+	bool bHideObjectsWithZeroEstimatedMemory = false; // hide objects with no Estimated Memory (System + Video)
+	bool bHideObjectsWithLowInclusiveEstimatedMemory = false; // hide objects with Inclusive Estimated Memory (System + Video) < 1 MiB
+	bool bHideObjectsWithLowEstimatedMemory = false; // hide objects with Exclusive Estimated Memory (System + Video) < 1 MiB
+	bool bHideObjectsWithLowImpact = false; // hide objects with Impact % < 1%
 
 	//////////////////////////////////////////////////
 
