[Insights] Normalized code across all four *ProfilerManager classes making them more consistent, cleaner, and slightly more future-proof:
- Member initializers moved to declaration;
- Constructor initializer list cleanup;
- Added MajorTabId member and GetMajorTabId() accessor to all managers;
- Added consistent bIsTabRegistered guard to all managers;
- Added IsAvailable() accessor to all managers;
- Updated UE_LOG to  UE_LOGF in these managers;
- Fixed many other smaller inconsistencies between managers code.

#rnx
#Synced-CL 52002978

diff --git a/Engine/Source/Developer/TraceInsights/Private/Insights/LoadingProfiler/LoadingProfilerManager.cpp b/Engine/Source/Developer/TraceInsights/Private/Insights/LoadingProfiler/LoadingProfilerManager.cpp
--- a/Engine/Source/Developer/TraceInsights/Private/Insights/LoadingProfiler/LoadingProfilerManager.cpp
+++ b/Engine/Source/Developer/TraceInsights/Private/Insights/LoadingProfiler/LoadingProfilerManager.cpp
@@ -26,9 +26,7 @@
 namespace UE::Insights::LoadingProfiler
 {
 
-////////////////////////////////////////////////////////////////////////////////////////////////////
-
-DEFINE_LOG_CATEGORY(LoadingProfiler);
+DEFINE_LOG_CATEGORY(LogLoadingProfiler);
 
 TSharedPtr<FLoadingProfilerManager> FLoadingProfilerManager::Instance = nullptr;
 
@@ -57,18 +55,10 @@ TSharedPtr<FLoadingProfilerManager> FLoadingProfilerManager::CreateInstance()
 ////////////////////////////////////////////////////////////////////////////////////////////////////
 
 FLoadingProfilerManager::FLoadingProfilerManager(TSharedRef<FUICommandList> InCommandList)
-	: bIsInitialized(false)
-	, bIsAvailable(false)
-	, CommandList(InCommandList)
+	: CommandList(InCommandList)
 	, ActionManager(this)
-	, ProfilerWindowWeakPtr()
-	, bIsTimingViewVisible(false)
-	, bIsEventAggregationTreeViewVisible(false)
-	, bIsObjectTypeAggregationTreeViewVisible(false)
-	, bIsPackageDetailsTreeViewVisible(false)
-	, bIsExportDetailsTreeViewVisible(false)
-	, bIsRequestsTreeViewVisible(false)
-	, LogListingName(TEXT("LoadingInsights"))
+	, MajorTabId(FInsightsManagerTabs::LoadingProfilerTabId)
+	, LogListingName("LoadingInsights")
 {
 }
 
@@ -83,7 +73,7 @@ void FLoadingProfilerManager::Initialize(IUnrealInsightsModule& InsightsModule)
 	}
 	bIsInitialized = true;
 
-	UE_LOG(LoadingProfiler, Log, TEXT("Initialize"));
+	UE_LOGF(LogLoadingProfiler, Log, "Initialize");
 
 	// Register tick functions.
 	OnTick = FTickerDelegate::CreateSP(this, &FLoadingProfilerManager::Tick);
@@ -125,7 +115,7 @@ void FLoadingProfilerManager::Shutdown()
 
 	FLoadingProfilerManager::Instance.Reset();
 
-	UE_LOG(LoadingProfiler, Log, TEXT("Shutdown"));
+	UE_LOGF(LogLoadingProfiler, Log, "Shutdown");
 }
 
 ////////////////////////////////////////////////////////////////////////////////////////////////////
@@ -151,19 +141,23 @@ void FLoadingProfilerManager::BindCommands()
 
 void FLoadingProfilerManager::RegisterMajorTabs(IUnrealInsightsModule& InsightsModule)
 {
-	const FInsightsMajorTabConfig& Config = InsightsModule.FindMajorTabConfig(FInsightsManagerTabs::LoadingProfilerTabId);
-
-	if (Config.bIsAvailable)
+	if (!bIsTabRegistered)
 	{
-		// Register tab spawner for the Asset Loading Insights.
-		FTabSpawnerEntry& TabSpawnerEntry = FGlobalTabmanager::Get()->RegisterNomadTabSpawner(FInsightsManagerTabs::LoadingProfilerTabId,
-			FOnSpawnTab::CreateRaw(this, &FLoadingProfilerManager::SpawnTab), FCanSpawnTab::CreateRaw(this, &FLoadingProfilerManager::CanSpawnTab))
-			.SetDisplayName(Config.TabLabel.IsSet() ? Config.TabLabel.GetValue() : LOCTEXT("LoadingProfilerTabTitle", "Asset Loading Insights"))
-			.SetTooltipText(Config.TabTooltip.IsSet() ? Config.TabTooltip.GetValue() : LOCTEXT("LoadingProfilerTooltipText", "Open the Asset Loading Insights tab."))
-			.SetIcon(Config.TabIcon.IsSet() ? Config.TabIcon.GetValue() : FSlateIcon(FInsightsStyle::GetStyleSetName(), "Icons.LoadingProfiler"));
-
-		TSharedRef<FWorkspaceItem> Group = Config.WorkspaceGroup.IsValid() ? Config.WorkspaceGroup.ToSharedRef() : FInsightsManager::Get()->GetInsightsMenuBuilder()->GetInsightsToolsGroup();
-		TabSpawnerEntry.SetGroup(Group);
+		bIsTabRegistered = true;
+		const FInsightsMajorTabConfig& Config = InsightsModule.FindMajorTabConfig(MajorTabId);
+		if (Config.bIsAvailable)
+		{
+			// Register tab spawner for the Asset Loading Insights.
+			FTabSpawnerEntry& TabSpawnerEntry = FGlobalTabmanager::Get()->RegisterNomadTabSpawner(MajorTabId,
+					FOnSpawnTab::CreateRaw(this, &FLoadingProfilerManager::SpawnTab),
+					FCanSpawnTab::CreateRaw(this, &FLoadingProfilerManager::CanSpawnTab))
+				.SetDisplayName(Config.TabLabel.IsSet() ? Config.TabLabel.GetValue() : LOCTEXT("LoadingProfilerTabTitle", "Asset Loading Insights"))
+				.SetTooltipText(Config.TabTooltip.IsSet() ? Config.TabTooltip.GetValue() : LOCTEXT("LoadingProfilerTooltipText", "Open the Asset Loading Insights tab."))
+				.SetIcon(Config.TabIcon.IsSet() ? Config.TabIcon.GetValue() : FSlateIcon(FInsightsStyle::GetStyleSetName(), "Icons.LoadingProfiler"));
+
+			TSharedRef<FWorkspaceItem> Group = Config.WorkspaceGroup.IsValid() ? Config.WorkspaceGroup.ToSharedRef() : FInsightsManager::Get()->GetInsightsMenuBuilder()->GetInsightsToolsGroup();
+			TabSpawnerEntry.SetGroup(Group);
+		}
 	}
 }
 
@@ -171,7 +165,8 @@ void FLoadingProfilerManager::RegisterMajorTabs(IUnrealInsightsModule& InsightsM
 
 void FLoadingProfilerManager::UnregisterMajorTabs()
 {
-	FGlobalTabmanager::Get()->UnregisterNomadTabSpawner(FInsightsManagerTabs::LoadingProfilerTabId);
+	FGlobalTabmanager::Get()->UnregisterNomadTabSpawner(MajorTabId);
+	bIsTabRegistered = false;
 }
 
 ////////////////////////////////////////////////////////////////////////////////////////////////////
@@ -181,7 +176,7 @@ TSharedRef<SDockTab> FLoadingProfilerManager::SpawnTab(const FSpawnTabArgs& Args
 	const TSharedRef<SDockTab> DockTab = SNew(SDockTab)
 		.TabRole(ETabRole::NomadTab);
 
-	// Register OnTabClosed to handle I/O profiler manager shutdown.
+	// Register OnTabClosed to handle Loading profiler manager shutdown.
 	DockTab->SetOnTabClosed(SDockTab::FOnTabClosedCallback::CreateRaw(this, &FLoadingProfilerManager::OnTabClosed));
 
 	// Create the SLoadingProfilerWindow widget.
@@ -204,6 +199,8 @@ bool FLoadingProfilerManager::CanSpawnTab(const FSpawnTabArgs& Args) const
 
 void FLoadingProfilerManager::OnTabClosed(TSharedRef<SDockTab> TabBeingClosed)
 {
+	OnWindowClosedEvent();
+
 	RemoveProfilerWindow();
 
 	// Disable TabClosed delegate.
@@ -240,7 +237,10 @@ bool FLoadingProfilerManager::Tick(float DeltaTime)
 	{
 		bool bIsProviderAvailable = false;
 
-		TSharedPtr<const TraceServices::IAnalysisSession> Session = FInsightsManager::Get()->GetSession();
+		TSharedPtr<FInsightsManager> InsightsManager = FInsightsManager::Get();
+		check(InsightsManager.IsValid());
+
+		TSharedPtr<const TraceServices::IAnalysisSession> Session = InsightsManager->GetSession();
 		if (Session.IsValid())
 		{
 			TraceServices::FAnalysisSessionReadScope SessionReadScope(*Session.Get());
@@ -271,11 +271,10 @@ bool FLoadingProfilerManager::Tick(float DeltaTime)
 			MessageLogModule.RegisterLogListing(GetLogListingName(), LOCTEXT("LoadingInsights", "Loading Insights"));
 			MessageLogModule.EnableMessageLogDisplay(true);
 
-			const FName& TabId = FInsightsManagerTabs::LoadingProfilerTabId;
-			if (FGlobalTabmanager::Get()->HasTabSpawner(TabId))
+			if (FGlobalTabmanager::Get()->HasTabSpawner(MajorTabId))
 			{
-				UE_LOG(LoadingProfiler, Log, TEXT("Opening the \"Asset Loading Insights\" tab..."));
-				FGlobalTabmanager::Get()->TryInvokeTab(TabId);
+				UE_LOGF(LogLoadingProfiler, Log, "Opening the \"Asset Loading Insights\" tab...");
+				FGlobalTabmanager::Get()->TryInvokeTab(MajorTabId);
 			}
 		}
 	}
@@ -287,7 +286,7 @@ bool FLoadingProfilerManager::Tick(float DeltaTime)
 
 void FLoadingProfilerManager::OnSessionChanged()
 {
-	UE_LOG(LoadingProfiler, Log, TEXT("OnSessionChanged"));
+	UE_LOGF(LogLoadingProfiler, Log, "OnSessionChanged");
 
 	bIsAvailable = false;
 	if (FInsightsManager::Get()->GetSession().IsValid())
@@ -300,7 +299,7 @@ void FLoadingProfilerManager::OnSessionChanged()
 	}
 
 	TSharedPtr<SLoadingProfilerWindow> Wnd = GetProfilerWindow();
-	if (Wnd.IsValid())
+	if (Wnd)
 	{
 		Wnd->Reset();
 	}
@@ -313,7 +312,7 @@ void FLoadingProfilerManager::ShowHideTimingView(const bool bIsVisible)
 	bIsTimingViewVisible = bIsVisible;
 
 	TSharedPtr<SLoadingProfilerWindow> Wnd = GetProfilerWindow();
-	if (Wnd.IsValid())
+	if (Wnd)
 	{
 		Wnd->ShowHideTab(FLoadingProfilerTabs::TimingViewID, bIsVisible);
 	}
@@ -326,7 +325,7 @@ void FLoadingProfilerManager::ShowHideEventAggregationTreeView(const bool bIsVis
 	bIsEventAggregationTreeViewVisible = bIsVisible;
 
 	TSharedPtr<SLoadingProfilerWindow> Wnd = GetProfilerWindow();
-	if (Wnd.IsValid())
+	if (Wnd)
 	{
 		Wnd->ShowHideTab(FLoadingProfilerTabs::EventAggregationTreeViewID, bIsVisible);
 	}
@@ -339,7 +338,7 @@ void FLoadingProfilerManager::ShowHideObjectTypeAggregationTreeView(const bool b
 	bIsObjectTypeAggregationTreeViewVisible = bIsVisible;
 
 	TSharedPtr<SLoadingProfilerWindow> Wnd = GetProfilerWindow();
-	if (Wnd.IsValid())
+	if (Wnd)
 	{
 		Wnd->ShowHideTab(FLoadingProfilerTabs::ObjectTypeAggregationTreeViewID, bIsVisible);
 	}
@@ -352,7 +351,7 @@ void FLoadingProfilerManager::ShowHidePackageDetailsTreeView(const bool bIsVisib
 	bIsPackageDetailsTreeViewVisible = bIsVisible;
 
 	TSharedPtr<SLoadingProfilerWindow> Wnd = GetProfilerWindow();
-	if (Wnd.IsValid())
+	if (Wnd)
 	{
 		Wnd->ShowHideTab(FLoadingProfilerTabs::PackageDetailsTreeViewID, bIsVisible);
 	}
@@ -365,7 +364,7 @@ void FLoadingProfilerManager::ShowHideExportDetailsTreeView(const bool bIsVisibl
 	bIsExportDetailsTreeViewVisible = bIsVisible;
 
 	TSharedPtr<SLoadingProfilerWindow> Wnd = GetProfilerWindow();
-	if (Wnd.IsValid())
+	if (Wnd)
 	{
 		Wnd->ShowHideTab(FLoadingProfilerTabs::ExportDetailsTreeViewID, bIsVisible);
 	}
@@ -378,7 +377,7 @@ void FLoadingProfilerManager::ShowHideRequestsTreeView(const bool bIsVisible)
 	bIsRequestsTreeViewVisible = bIsVisible;
 
 	TSharedPtr<SLoadingProfilerWindow> Wnd = GetProfilerWindow();
-	if (Wnd.IsValid())
+	if (Wnd)
 	{
 		Wnd->ShowHideTab(FLoadingProfilerTabs::RequestsTreeViewID, bIsVisible);
 	}
@@ -392,7 +391,7 @@ void FLoadingProfilerManager::OnWindowClosedEvent()
 	if (Wnd)
 	{
 		TSharedPtr<TimingProfiler::STimingView> TimingView = Wnd->GetTimingView();
-		if (TimingView.IsValid())
+		if (TimingView)
 		{
 			TimingView->CloseQuickFindTab();
 		}
diff --git a/Engine/Source/Developer/TraceInsights/Private/Insights/LoadingProfiler/LoadingProfilerManager.h b/Engine/Source/Developer/TraceInsights/Private/Insights/LoadingProfiler/LoadingProfilerManager.h
--- a/Engine/Source/Developer/TraceInsights/Private/Insights/LoadingProfiler/LoadingProfilerManager.h
+++ b/Engine/Source/Developer/TraceInsights/Private/Insights/LoadingProfiler/LoadingProfilerManager.h
@@ -21,7 +21,7 @@ namespace UE::Insights::LoadingProfiler
 
 class SLoadingProfilerWindow;
 
-DECLARE_LOG_CATEGORY_EXTERN(LoadingProfiler, Log, All);
+DECLARE_LOG_CATEGORY_EXTERN(LogLoadingProfiler, Log, All);
 
 ////////////////////////////////////////////////////////////////////////////////////////////////////
 /**
@@ -50,6 +50,8 @@ public:
 	 */
 	static TSharedPtr<FLoadingProfilerManager> Get();
 
+	bool IsAvailable() const { return bIsAvailable; }
+
 	//////////////////////////////////////////////////
 	// IInsightsComponent
 
@@ -79,6 +81,9 @@ public:
 		return ProfilerWindowWeakPtr.Pin();
 	}
 
+	const FName& GetMajorTabId() const { return MajorTabId; }
+	const FName& GetLogListingName() const { return LogListingName; }
+
 	////////////////////////////////////////////////////////////////////////////////////////////////////
 	// Getters and setters used by Toggle Commands.
 
@@ -107,7 +112,7 @@ public:
 	void SetExportDetailsTreeViewVisible(const bool bIsVisible) { bIsExportDetailsTreeViewVisible = bIsVisible; }
 	void ShowHideExportDetailsTreeView(const bool bIsVisible);
 
-	/** @return true, if the Export Details tree view is visible */
+	/** @return true, if the Requests tree view is visible */
 	const bool IsRequestsTreeViewVisible() const { return bIsRequestsTreeViewVisible; }
 	void SetRequestsTreeViewVisible(const bool bIsVisible) { bIsRequestsTreeViewVisible = bIsVisible; }
 	void ShowHideRequestsTreeView(const bool bIsVisible);
@@ -116,8 +121,6 @@ public:
 
 	void OnSessionChanged();
 
-	const FName& GetLogListingName() const { return LogListingName; }
-
 private:
 	/** Binds our UI commands to delegates. */
 	void BindCommands();
@@ -144,8 +147,12 @@ private:
 	}
 
 private:
-	bool bIsInitialized;
-	bool bIsAvailable;
+	bool bIsInitialized = false;
+	bool bIsAvailable = false;
+
+	/** If the tab is registered in the tab manager. */
+	bool bIsTabRegistered = false;
+
 	FAvailabilityCheck AvailabilityCheck;
 
 	/** The delegate to be invoked when this manager ticks. */
@@ -164,30 +171,31 @@ private:
 	TWeakPtr<SLoadingProfilerWindow> ProfilerWindowWeakPtr;
 
 	/** If the Timing view is visible or hidden. */
-	bool bIsTimingViewVisible;
+	bool bIsTimingViewVisible = false;
 
 	/** If the Event Aggregation tree view is visible or hidden. */
-	bool bIsEventAggregationTreeViewVisible;
+	bool bIsEventAggregationTreeViewVisible = false;
 
 	/** If the Object Type Aggregation tree view is visible or hidden. */
-	bool bIsObjectTypeAggregationTreeViewVisible;
+	bool bIsObjectTypeAggregationTreeViewVisible = false;
 
 	/** If the Package Details tree view is visible or hidden. */
-	bool bIsPackageDetailsTreeViewVisible;
+	bool bIsPackageDetailsTreeViewVisible = false;
 
 	/** If the Export Details tree view is visible or hidden. */
-	bool bIsExportDetailsTreeViewVisible;
+	bool bIsExportDetailsTreeViewVisible = false;
 
 	/** If the Requests tree view is visible or hidden. */
-	bool bIsRequestsTreeViewVisible;
+	bool bIsRequestsTreeViewVisible = false;
+
+	/** The id of the Asset Loading Insights major tab. */
+	FName MajorTabId;
 
-	/** The name of the Loading Insights log listing. */
+	/** The name of the Loading Profiler log listing. */
 	FName LogListingName;
 
 	/** A shared pointer to the global instance of the Loading Profiler manager. */
 	static TSharedPtr<FLoadingProfilerManager> Instance;
 };
 
-////////////////////////////////////////////////////////////////////////////////////////////////////
-
 } // namespace UE::Insights::LoadingProfiler
diff --git a/Engine/Source/Developer/TraceInsights/Private/Insights/MemoryProfiler/MemoryProfilerManager.cpp b/Engine/Source/Developer/TraceInsights/Private/Insights/MemoryProfiler/MemoryProfilerManager.cpp
--- a/Engine/Source/Developer/TraceInsights/Private/Insights/MemoryProfiler/MemoryProfilerManager.cpp
+++ b/Engine/Source/Developer/TraceInsights/Private/Insights/MemoryProfiler/MemoryProfilerManager.cpp
@@ -70,16 +70,10 @@ TSharedPtr<FMemoryProfilerManager> FMemoryProfilerManager::CreateInstance()
 ////////////////////////////////////////////////////////////////////////////////////////////////////
 
 FMemoryProfilerManager::FMemoryProfilerManager(TSharedRef<FUICommandList> InCommandList)
-	: bIsInitialized(false)
-	, bIsAvailable(false)
-	, CommandList(InCommandList)
+	: CommandList(InCommandList)
 	, ActionManager(this)
-	, ProfilerWindowWeakPtr()
-	, bIsTimingViewVisible(false)
-	, bIsMemInvestigationViewVisible(false)
-	, bIsMemTagTreeViewVisible(false)
-	, bIsModulesViewVisible(false)
-	, LogListingName(TEXT("MemoryInsights"))
+	, MajorTabId(FInsightsManagerTabs::MemoryProfilerTabId)
+	, LogListingName("MemoryInsights")
 {
 }
 
@@ -94,7 +88,7 @@ void FMemoryProfilerManager::Initialize(IUnrealInsightsModule& InsightsModule)
 	}
 	bIsInitialized = true;
 
-	UE_LOG(LogMemoryProfiler, Log, TEXT("Initialize"));
+	UE_LOGF(LogMemoryProfiler, Log, "Initialize");
 
 	// Register tick functions.
 	OnTick = FTickerDelegate::CreateSP(this, &FMemoryProfilerManager::Tick);
@@ -136,7 +130,7 @@ void FMemoryProfilerManager::Shutdown()
 
 	FMemoryProfilerManager::Instance.Reset();
 
-	UE_LOG(LogMemoryProfiler, Log, TEXT("Shutdown"));
+	UE_LOGF(LogMemoryProfiler, Log, "Shutdown");
 }
 
 ////////////////////////////////////////////////////////////////////////////////////////////////////
@@ -173,13 +167,13 @@ void FMemoryProfilerManager::RegisterMajorTabs(IUnrealInsightsModule& InsightsMo
 	if (!bIsTabRegistered && bShouldRegister)
 	{
 		bIsTabRegistered = true;
-		const FInsightsMajorTabConfig& Config = InsightsModule.FindMajorTabConfig(FInsightsManagerTabs::MemoryProfilerTabId);
-
+		const FInsightsMajorTabConfig& Config = InsightsModule.FindMajorTabConfig(MajorTabId);
 		if (Config.bIsAvailable)
 		{
 			// Register tab spawner for the Memory Insights.
-			FTabSpawnerEntry& TabSpawnerEntry = FGlobalTabmanager::Get()->RegisterNomadTabSpawner(FInsightsManagerTabs::MemoryProfilerTabId,
-				FOnSpawnTab::CreateRaw(this, &FMemoryProfilerManager::SpawnTab), FCanSpawnTab::CreateRaw(this, &FMemoryProfilerManager::CanSpawnTab))
+			FTabSpawnerEntry& TabSpawnerEntry = FGlobalTabmanager::Get()->RegisterNomadTabSpawner(MajorTabId,
+					FOnSpawnTab::CreateRaw(this, &FMemoryProfilerManager::SpawnTab),
+					FCanSpawnTab::CreateRaw(this, &FMemoryProfilerManager::CanSpawnTab))
 				.SetDisplayName(Config.TabLabel.IsSet() ? Config.TabLabel.GetValue() : LOCTEXT("MemoryProfilerTabTitle", "Memory Insights"))
 				.SetTooltipText(Config.TabTooltip.IsSet() ? Config.TabTooltip.GetValue() : LOCTEXT("MemoryProfilerTooltipText", "Open the Memory Insights tab."))
 				.SetIcon(Config.TabIcon.IsSet() ? Config.TabIcon.GetValue() : FSlateIcon(FInsightsStyle::GetStyleSetName(), "Icons.MemoryProfiler"));
@@ -194,7 +188,8 @@ void FMemoryProfilerManager::RegisterMajorTabs(IUnrealInsightsModule& InsightsMo
 
 void FMemoryProfilerManager::UnregisterMajorTabs()
 {
-	FGlobalTabmanager::Get()->UnregisterNomadTabSpawner(FInsightsManagerTabs::MemoryProfilerTabId);
+	FGlobalTabmanager::Get()->UnregisterNomadTabSpawner(MajorTabId);
+	bIsTabRegistered = false;
 }
 
 ////////////////////////////////////////////////////////////////////////////////////////////////////
@@ -204,7 +199,7 @@ TSharedRef<SDockTab> FMemoryProfilerManager::SpawnTab(const FSpawnTabArgs& Args)
 	const TSharedRef<SDockTab> DockTab = SNew(SDockTab)
 		.TabRole(ETabRole::NomadTab);
 
-	// Register OnTabClosed to handle I/O profiler manager shutdown.
+	// Register OnTabClosed to handle Memory profiler manager shutdown.
 	DockTab->SetOnTabClosed(SDockTab::FOnTabClosedCallback::CreateRaw(this, &FMemoryProfilerManager::OnTabClosed));
 
 	// Create the SMemoryProfilerWindow widget.
@@ -231,13 +226,7 @@ bool FMemoryProfilerManager::CanSpawnTab(const FSpawnTabArgs& Args) const
 
 void FMemoryProfilerManager::OnTabClosed(TSharedRef<SDockTab> TabBeingClosed)
 {
-	{
-		TSharedPtr<SMemoryProfilerWindow> ProfilerWindow = GetProfilerWindow();
-		if (ProfilerWindow.IsValid())
-		{
-			ProfilerWindow->CloseMemAllocTableTreeTabs();
-		}
-	}
+	OnWindowClosedEvent();
 
 	RemoveProfilerWindow();
 
@@ -286,12 +275,15 @@ bool FMemoryProfilerManager::Tick(float DeltaTime)
 	}
 #endif
 
-	// Check if session has Memory events (to spawn the tab), but not too often.
+	// Check if the session has Memory events (to spawn the tab), but not too often.
 	if (!bIsAvailable && AvailabilityCheck.Tick())
 	{
 		bool bShouldBeAvailable = false;
 
-		TSharedPtr<const TraceServices::IAnalysisSession> Session = FInsightsManager::Get()->GetSession();
+		TSharedPtr<FInsightsManager> InsightsManager = FInsightsManager::Get();
+		check(InsightsManager.IsValid());
+
+		TSharedPtr<const TraceServices::IAnalysisSession> Session = InsightsManager->GetSession();
 		if (Session.IsValid())
 		{
 			{
@@ -339,11 +331,10 @@ bool FMemoryProfilerManager::Tick(float DeltaTime)
 			MessageLogModule.EnableMessageLogDisplay(true);
 
 #if !WITH_EDITOR
-			const FName& TabId = FInsightsManagerTabs::MemoryProfilerTabId;
-			if (FGlobalTabmanager::Get()->HasTabSpawner(TabId))
+			if (FGlobalTabmanager::Get()->HasTabSpawner(MajorTabId))
 			{
-				UE_LOG(LogMemoryProfiler, Log, TEXT("Opening the \"Memory Insights\" tab..."));
-				FGlobalTabmanager::Get()->TryInvokeTab(TabId);
+				UE_LOGF(LogMemoryProfiler, Log, "Opening the \"Memory Insights\" tab...");
+				FGlobalTabmanager::Get()->TryInvokeTab(MajorTabId);
 			}
 #endif
 		}
@@ -362,7 +353,7 @@ bool FMemoryProfilerManager::Tick(float DeltaTime)
 
 void FMemoryProfilerManager::OnSessionChanged()
 {
-	UE_LOG(LogMemoryProfiler, Log, TEXT("OnSessionChanged"));
+	UE_LOGF(LogMemoryProfiler, Log, "OnSessionChanged");
 
 	bIsAvailable = false;
 	if (FInsightsManager::Get()->GetSession().IsValid())
@@ -375,7 +366,7 @@ void FMemoryProfilerManager::OnSessionChanged()
 	}
 
 	TSharedPtr<SMemoryProfilerWindow> Wnd = GetProfilerWindow();
-	if (Wnd.IsValid())
+	if (Wnd)
 	{
 		Wnd->Reset();
 	}
@@ -388,7 +379,7 @@ void FMemoryProfilerManager::ShowHideTimingView(const bool bIsVisible)
 	bIsTimingViewVisible = bIsVisible;
 
 	TSharedPtr<SMemoryProfilerWindow> Wnd = GetProfilerWindow();
-	if (Wnd.IsValid())
+	if (Wnd)
 	{
 		Wnd->ShowHideTab(FMemoryProfilerTabs::TimingViewID, bIsVisible);
 	}
@@ -401,7 +392,7 @@ void FMemoryProfilerManager::ShowHideMemInvestigationView(const bool bIsVisible)
 	bIsMemInvestigationViewVisible = bIsVisible;
 
 	TSharedPtr<SMemoryProfilerWindow> Wnd = GetProfilerWindow();
-	if (Wnd.IsValid())
+	if (Wnd)
 	{
 		Wnd->ShowHideTab(FMemoryProfilerTabs::MemInvestigationViewID, bIsVisible);
 	}
@@ -414,7 +405,7 @@ void FMemoryProfilerManager::ShowHideMemTagTreeView(const bool bIsVisible)
 	bIsMemTagTreeViewVisible = bIsVisible;
 
 	TSharedPtr<SMemoryProfilerWindow> Wnd = GetProfilerWindow();
-	if (Wnd.IsValid())
+	if (Wnd)
 	{
 		Wnd->ShowHideTab(FMemoryProfilerTabs::MemTagTreeViewID, bIsVisible);
 	}
@@ -427,7 +418,7 @@ void FMemoryProfilerManager::ShowHideModulesView(const bool bIsVisible)
 	bIsModulesViewVisible = bIsVisible;
 
 	TSharedPtr<SMemoryProfilerWindow> Wnd = GetProfilerWindow();
-	if (Wnd.IsValid())
+	if (Wnd)
 	{
 		Wnd->ShowHideTab(FMemoryProfilerTabs::ModulesViewID, bIsVisible);
 	}
@@ -437,14 +428,14 @@ void FMemoryProfilerManager::ShowHideModulesView(const bool bIsVisible)
 
 void FMemoryProfilerManager::OnWindowClosedEvent()
 {
-	// Need to close MemAlloc window to prevent it being saved in the layout and spawning as an "Unregister Tab" on the next application start.
 	TSharedPtr<SMemoryProfilerWindow> Wnd = GetProfilerWindow();
-	if (Wnd.IsValid())
+	if (Wnd)
 	{
+		// Need to close MemAlloc window to prevent it being saved in the layout and spawning as an "Unregister Tab" on the next application start.
 		Wnd->CloseMemAllocTableTreeTabs();
 
 		TSharedPtr<TimingProfiler::STimingView> TimingView = Wnd->GetTimingView();
-		if (TimingView.IsValid())
+		if (TimingView)
 		{
 			TimingView->CloseQuickFindTab();
 		}
diff --git a/Engine/Source/Developer/TraceInsights/Private/Insights/MemoryProfiler/MemoryProfilerManager.h b/Engine/Source/Developer/TraceInsights/Private/Insights/MemoryProfiler/MemoryProfilerManager.h
--- a/Engine/Source/Developer/TraceInsights/Private/Insights/MemoryProfiler/MemoryProfilerManager.h
+++ b/Engine/Source/Developer/TraceInsights/Private/Insights/MemoryProfiler/MemoryProfilerManager.h
@@ -6,8 +6,8 @@
 
 #include "Containers/Ticker.h"
 #include "Framework/Commands/UICommandList.h"
-#include "Templates/SharedPointer.h"
 #include "Logging/LogMacros.h"
+#include "Templates/SharedPointer.h"
 
 // TraceInsightsCore
 #include "InsightsCore/Common/AvailabilityCheck.h"
@@ -51,6 +51,8 @@ public:
 	 */
 	static TSharedPtr<FMemoryProfilerManager> Get();
 
+	bool IsAvailable() const { return bIsAvailable; }
+
 	//////////////////////////////////////////////////
 	// IInsightsComponent
 
@@ -58,8 +60,8 @@ public:
 	virtual void Shutdown() override;
 	virtual void RegisterMajorTabs(IUnrealInsightsModule& InsightsModule) override;
 	virtual void UnregisterMajorTabs() override;
-	virtual bool Exec(const TCHAR* Cmd, FOutputDevice& Ar) override;
 	virtual void OnWindowClosedEvent() override;
+	virtual bool Exec(const TCHAR* Cmd, FOutputDevice& Ar) override;
 
 	//////////////////////////////////////////////////
 
@@ -81,6 +83,9 @@ public:
 		return ProfilerWindowWeakPtr.Pin();
 	}
 
+	const FName& GetMajorTabId() const { return MajorTabId; }
+	const FName& GetLogListingName() const { return LogListingName; }
+
 	FMemorySharedState* GetSharedState();
 
 	////////////////////////////////////////////////////////////////////////////////////////////////////
@@ -104,8 +109,6 @@ public:
 
 	////////////////////////////////////////////////////////////////////////////////////////////////////
 
-	const FName& GetLogListingName() const { return LogListingName; }
-
 	void OnSessionChanged();
 
 private:
@@ -134,8 +137,12 @@ private:
 	}
 
 private:
-	bool bIsInitialized;
-	bool bIsAvailable;
+	bool bIsInitialized = false;
+	bool bIsAvailable = false;
+
+	/** If the tab is registered in the tab manager. */
+	bool bIsTabRegistered = false;
+
 	FAvailabilityCheck AvailabilityCheck;
 
 	/** The delegate to be invoked when this manager ticks. */
@@ -154,25 +161,25 @@ private:
 	TWeakPtr<SMemoryProfilerWindow> ProfilerWindowWeakPtr;
 
 	/** If the Timing view is visible or hidden. */
-	bool bIsTimingViewVisible;
+	bool bIsTimingViewVisible = false;
 
 	/** If the Memory Investigation (Alloc Queries) view is visible or hidden. */
-	bool bIsMemInvestigationViewVisible;
+	bool bIsMemInvestigationViewVisible = false;
 
 	/** If the Memory Tags tree view is visible or hidden. */
-	bool bIsMemTagTreeViewVisible;
+	bool bIsMemTagTreeViewVisible = false;
 
 	/** If the Modules view is visible or hidden. */
-	bool bIsModulesViewVisible;
+	bool bIsModulesViewVisible = false;
 
-	/** The name of the Memory Insights log listing. */
+	/** The id of the Memory Insights major tab. */
+	FName MajorTabId;
+
+	/** The name of the Memory Profiler log listing. */
 	FName LogListingName;
 
 	/** A shared pointer to the global instance of the Memory Profiler manager. */
 	static TSharedPtr<FMemoryProfilerManager> Instance;
-
-	/** If the tab is registered in the tab manager. */
-	bool bIsTabRegistered = false;
 };
 
 } // namespace UE::Insights::MemoryProfiler
diff --git a/Engine/Source/Developer/TraceInsights/Private/Insights/NetworkingProfiler/NetworkingProfilerManager.cpp b/Engine/Source/Developer/TraceInsights/Private/Insights/NetworkingProfiler/NetworkingProfilerManager.cpp
--- a/Engine/Source/Developer/TraceInsights/Private/Insights/NetworkingProfiler/NetworkingProfilerManager.cpp
+++ b/Engine/Source/Developer/TraceInsights/Private/Insights/NetworkingProfiler/NetworkingProfilerManager.cpp
@@ -48,12 +48,10 @@ TSharedPtr<FNetworkingProfilerManager> FNetworkingProfilerManager::CreateInstanc
 ////////////////////////////////////////////////////////////////////////////////////////////////////
 
 FNetworkingProfilerManager::FNetworkingProfilerManager(TSharedRef<FUICommandList> InCommandList)
-	: bIsInitialized(false)
-	, bIsAvailable(false)
-	, CommandList(InCommandList)
+	: CommandList(InCommandList)
 	, ActionManager(this)
-	, ProfilerWindows()
-	, LogListingName(TEXT("NetworkingInsights"))
+	, MajorTabId(FInsightsManagerTabs::NetworkingProfilerTabId)
+	, LogListingName("NetworkingInsights")
 {
 }
 
@@ -68,7 +66,7 @@ void FNetworkingProfilerManager::Initialize(IUnrealInsightsModule& InsightsModul
 	}
 	bIsInitialized = true;
 
-	UE_LOG(LogNetworkingProfiler, Log, TEXT("Initialize"));
+	UE_LOGF(LogNetworkingProfiler, Log, "Initialize");
 
 	// Register tick functions.
 	OnTick = FTickerDelegate::CreateSP(this, &FNetworkingProfilerManager::Tick);
@@ -110,7 +108,7 @@ void FNetworkingProfilerManager::Shutdown()
 
 	FNetworkingProfilerManager::Instance.Reset();
 
-	UE_LOG(LogNetworkingProfiler, Log, TEXT("Shutdown"));
+	UE_LOGF(LogNetworkingProfiler, Log, "Shutdown");
 }
 
 ////////////////////////////////////////////////////////////////////////////////////////////////////
@@ -135,17 +133,16 @@ static TSharedPtr<SDockTab> NeverReuse(const FTabId&)
 
 void FNetworkingProfilerManager::RegisterMajorTabs(IUnrealInsightsModule& InsightsModule)
 {
-	const FInsightsMajorTabConfig& Config = InsightsModule.FindMajorTabConfig(FInsightsManagerTabs::NetworkingProfilerTabId);
-
-	if (Config.bIsAvailable)
+	if (!bIsTabRegistered)
 	{
-		// Register tab spawner(s) for the Networking Insights.
-		//for (int32 ReservedId = 0; ReservedId < 10; ++ReservedId)
+		bIsTabRegistered = true;
+		const FInsightsMajorTabConfig& Config = InsightsModule.FindMajorTabConfig(MajorTabId);
+		if (Config.bIsAvailable)
 		{
-			FName TabId = FInsightsManagerTabs::NetworkingProfilerTabId;
-			//TabId.SetNumber(ReservedId);
-			FTabSpawnerEntry& TabSpawnerEntry = FGlobalTabmanager::Get()->RegisterNomadTabSpawner(TabId,
-				FOnSpawnTab::CreateRaw(this, &FNetworkingProfilerManager::SpawnTab), FCanSpawnTab::CreateRaw(this, &FNetworkingProfilerManager::CanSpawnTab))
+			// Register tab spawner for the Networking Insights.
+			FTabSpawnerEntry& TabSpawnerEntry = FGlobalTabmanager::Get()->RegisterNomadTabSpawner(MajorTabId,
+					FOnSpawnTab::CreateRaw(this, &FNetworkingProfilerManager::SpawnTab),
+					FCanSpawnTab::CreateRaw(this, &FNetworkingProfilerManager::CanSpawnTab))
 				.SetReuseTabMethod(FOnFindTabToReuse::CreateStatic(&NeverReuse))
 				.SetDisplayName(Config.TabLabel.IsSet() ? Config.TabLabel.GetValue() : LOCTEXT("NetworkingProfilerTabTitle", "Networking Insights"))
 				.SetTooltipText(Config.TabTooltip.IsSet() ? Config.TabTooltip.GetValue() : LOCTEXT("NetworkingProfilerTooltipText", "Open the Networking Insights tab."))
@@ -161,12 +158,8 @@ void FNetworkingProfilerManager::RegisterMajorTabs(IUnrealInsightsModule& Insigh
 
 void FNetworkingProfilerManager::UnregisterMajorTabs()
 {
-	//for (int32 ReservedId = 0; ReservedId < 10; ++ReservedId)
-	{
-		FName TabId = FInsightsManagerTabs::NetworkingProfilerTabId;
-		//TabId.SetNumber(ReservedId);
-		FGlobalTabmanager::Get()->UnregisterNomadTabSpawner(TabId);
-	}
+	FGlobalTabmanager::Get()->UnregisterNomadTabSpawner(MajorTabId);
+	bIsTabRegistered = false;
 }
 
 ////////////////////////////////////////////////////////////////////////////////////////////////////
@@ -176,7 +169,7 @@ TSharedRef<SDockTab> FNetworkingProfilerManager::SpawnTab(const FSpawnTabArgs& A
 	const TSharedRef<SDockTab> DockTab = SNew(SDockTab)
 		.TabRole(ETabRole::NomadTab);
 
-	// Register OnTabClosed to handle I/O profiler manager shutdown.
+	// Register OnTabClosed to handle Networking profiler manager shutdown.
 	DockTab->SetOnTabClosed(SDockTab::FOnTabClosedCallback::CreateRaw(this, &FNetworkingProfilerManager::OnTabClosed));
 
 	// Create the SNetworkingProfilerWindow widget.
@@ -237,7 +230,10 @@ bool FNetworkingProfilerManager::Tick(float DeltaTime)
 	{
 		uint32 NetTraceVersion = 0;
 
-		TSharedPtr<const TraceServices::IAnalysisSession> Session = FInsightsManager::Get()->GetSession();
+		TSharedPtr<FInsightsManager> InsightsManager = FInsightsManager::Get();
+		check(InsightsManager.IsValid());
+
+		TSharedPtr<const TraceServices::IAnalysisSession> Session = InsightsManager->GetSession();
 		if (Session.IsValid())
 		{
 			TraceServices::FAnalysisSessionReadScope SessionReadScope(*Session.Get());
@@ -264,35 +260,17 @@ bool FNetworkingProfilerManager::Tick(float DeltaTime)
 		{
 			bIsAvailable = true;
 
-			const FName& TabId = FInsightsManagerTabs::NetworkingProfilerTabId;
-			if (FGlobalTabmanager::Get()->HasTabSpawner(TabId))
-			{
-				// Spawn 2 tabs.
-				UE_LOG(LogNetworkingProfiler, Log, TEXT("Opening the \"Networking Insights\" tabs..."));
-				FGlobalTabmanager::Get()->TryInvokeTab(TabId);
-				FGlobalTabmanager::Get()->TryInvokeTab(TabId);
-			}
-
 			FMessageLogModule& MessageLogModule = FModuleManager::LoadModuleChecked<FMessageLogModule>("MessageLog");
-			MessageLogModule.RegisterLogListing(GetLogListingName(), LOCTEXT("NetworkingInsights", "Networking Profiler Insights"));
+			MessageLogModule.RegisterLogListing(GetLogListingName(), LOCTEXT("NetworkingInsights", "Networking Insights"));
 			MessageLogModule.EnableMessageLogDisplay(true);
 
-			//int32 SpawnTabCount = 2; // we want to spawn 2 tabs
-			//for (int32 ReservedId = 0; SpawnTabCount > 0 && ReservedId < 10; ++ReservedId)
-			//{
-			//	FName TabId = FInsightsManagerTabs::NetworkingProfilerTabId;
-			//	TabId.SetNumber(ReservedId);
-			//
-			//	if (FGlobalTabmanager::Get()->HasTabSpawner(TabId) && 
-			//		!FGlobalTabmanager::Get()->FindExistingLiveTab(TabId).IsValid())
-			//	{
-			//		UE_LOG(LogNetworkingProfiler, Log, TEXT("Opening the \"Networking Insights\" tab..."));
-			//		FGlobalTabmanager::Get()->TryInvokeTab(TabId);
-			//		--SpawnTabCount;
-			//	}
-			//}
-
-			// ActivateTimingInsightsTab();
+			if (FGlobalTabmanager::Get()->HasTabSpawner(MajorTabId))
+			{
+				// Spawn 2 tabs.
+				UE_LOGF(LogNetworkingProfiler, Log, "Opening the \"Networking Insights\" tabs...");
+				FGlobalTabmanager::Get()->TryInvokeTab(MajorTabId);
+				FGlobalTabmanager::Get()->TryInvokeTab(MajorTabId);
+			}
 		}
 	}
 
@@ -303,7 +281,7 @@ bool FNetworkingProfilerManager::Tick(float DeltaTime)
 
 void FNetworkingProfilerManager::OnSessionChanged()
 {
-	UE_LOG(LogNetworkingProfiler, Log, TEXT("OnSessionChanged"));
+	UE_LOGF(LogNetworkingProfiler, Log, "OnSessionChanged");
 
 	bIsAvailable = false;
 	if (FInsightsManager::Get()->GetSession().IsValid())
@@ -318,7 +296,7 @@ void FNetworkingProfilerManager::OnSessionChanged()
 	for (TWeakPtr<SNetworkingProfilerWindow> WndWeakPtr : ProfilerWindows)
 	{
 		TSharedPtr<SNetworkingProfilerWindow> Wnd = WndWeakPtr.Pin();
-		if (Wnd.IsValid())
+		if (Wnd)
 		{
 			Wnd->Reset();
 		}
diff --git a/Engine/Source/Developer/TraceInsights/Private/Insights/NetworkingProfiler/NetworkingProfilerManager.h b/Engine/Source/Developer/TraceInsights/Private/Insights/NetworkingProfiler/NetworkingProfilerManager.h
--- a/Engine/Source/Developer/TraceInsights/Private/Insights/NetworkingProfiler/NetworkingProfilerManager.h
+++ b/Engine/Source/Developer/TraceInsights/Private/Insights/NetworkingProfiler/NetworkingProfilerManager.h
@@ -2,12 +2,12 @@
 
 #pragma once
 
-#include "CoreMinimal.h"
+#include "CoreTypes.h"
 
 #include "Containers/Ticker.h"
 #include "Framework/Commands/UICommandList.h"
-#include "Templates/SharedPointer.h"
 #include "Logging/LogMacros.h"
+#include "Templates/SharedPointer.h"
 
 // TraceInsightsCore
 #include "InsightsCore/Common/AvailabilityCheck.h"
@@ -50,6 +50,8 @@ public:
 	 */
 	static TSharedPtr<FNetworkingProfilerManager> Get();
 
+	bool IsAvailable() const { return bIsAvailable; }
+
 	//////////////////////////////////////////////////
 	// IInsightsComponent
 
@@ -88,10 +90,11 @@ public:
 		return ProfilerWindows[Index].Pin();
 	}
 
-	void OnSessionChanged();
-
+	const FName& GetMajorTabId() const { return MajorTabId; }
 	const FName& GetLogListingName() const { return LogListingName; }
 
+	void OnSessionChanged();
+
 private:
 	/** Binds our UI commands to delegates. */
 	void BindCommands();
@@ -118,8 +121,12 @@ private:
 	}
 
 private:
-	bool bIsInitialized;
-	bool bIsAvailable;
+	bool bIsInitialized = false;
+	bool bIsAvailable = false;
+
+	/** If the tab is registered in the tab manager. */
+	bool bIsTabRegistered = false;
+
 	FAvailabilityCheck AvailabilityCheck;
 
 	/** The delegate to be invoked when this manager ticks. */
@@ -137,11 +144,14 @@ private:
 	/** A list of weak pointers to the Networking Profiler windows. */
 	TArray<TWeakPtr<class SNetworkingProfilerWindow>> ProfilerWindows;
 
-	/** A shared pointer to the global instance of the NetworkingProfiler manager. */
-	static TSharedPtr<FNetworkingProfilerManager> Instance;
+	/** The id of the Networking Insights major tab. */
+	FName MajorTabId;
 
 	/** The name of the Networking Profiler log listing. */
 	FName LogListingName;
+
+	/** A shared pointer to the global instance of the Networking Profiler manager. */
+	static TSharedPtr<FNetworkingProfilerManager> Instance;
 };
 
 } // namespace UE::Insights::NetworkingProfiler
diff --git a/Engine/Source/Developer/TraceInsights/Private/Insights/TimingProfiler/TimingProfilerManager.cpp b/Engine/Source/Developer/TraceInsights/Private/Insights/TimingProfiler/TimingProfilerManager.cpp
--- a/Engine/Source/Developer/TraceInsights/Private/Insights/TimingProfiler/TimingProfilerManager.cpp
+++ b/Engine/Source/Developer/TraceInsights/Private/Insights/TimingProfiler/TimingProfilerManager.cpp
@@ -74,24 +74,12 @@ TSharedPtr<FTimingProfilerManager> FTimingProfilerManager::CreateInstance()
 ////////////////////////////////////////////////////////////////////////////////////////////////////
 
 FTimingProfilerManager::FTimingProfilerManager(TSharedRef<FUICommandList> InCommandList)
-	: bIsInitialized(false)
-	, bIsAvailable(false)
-	, CommandList(InCommandList)
+	: CommandList(InCommandList)
 	, ActionManager(this)
-	, ProfilerWindowWeakPtr()
 	, Filter(MakeShared<FLogFilter>())
-	, bIsFramesTrackVisible(false)
-	, bIsTimingViewVisible(false)
-	, bIsTimersViewVisible(false)
-	, bIsCallersTreeViewVisible(false)
-	, bIsCalleesTreeViewVisible(false)
-	, bIsStatsCountersViewVisible(false)
-	, bIsLogViewVisible(false)
-	, SelectionStartTime(0.0)
-	, SelectionEndTime(0.0)
-	, SelectedTimerId(InvalidTimerId)
 	, TimerButterflyAggregator(MakeShared<FTimerButterflyAggregator>())
-	, LogListingName(TEXT("TimingInsights"))
+	, MajorTabId(FInsightsManagerTabs::TimingProfilerTabId)
+	, LogListingName("TimingInsights")
 {
 }
 
@@ -106,7 +94,7 @@ void FTimingProfilerManager::Initialize(IUnrealInsightsModule& InsightsModule)
 	}
 	bIsInitialized = true;
 
-	UE_LOG(LogTimingProfiler, Log, TEXT("Initialize"));
+	UE_LOGF(LogTimingProfiler, Log, "Initialize");
 
 	// Register tick functions.
 	OnTick = FTickerDelegate::CreateSP(this, &FTimingProfilerManager::Tick);
@@ -119,7 +107,7 @@ void FTimingProfilerManager::Initialize(IUnrealInsightsModule& InsightsModule)
 	FTimingProfilerCommands::Register();
 	BindCommands();
 
-	InsightsModule.OnRegisterMajorTabExtension(FInsightsManagerTabs::TimingProfilerTabId);
+	InsightsModule.OnRegisterMajorTabExtension(MajorTabId);
 
 	FInsightsManager::Get()->GetSessionChangedEvent().AddSP(this, &FTimingProfilerManager::OnSessionChanged);
 	OnSessionChanged();
@@ -156,7 +144,7 @@ void FTimingProfilerManager::Shutdown()
 
 	FTimingProfilerManager::Instance.Reset();
 
-	UE_LOG(LogTimingProfiler, Log, TEXT("Shutdown"));
+	UE_LOGF(LogTimingProfiler, Log, "Shutdown");
 }
 
 ////////////////////////////////////////////////////////////////////////////////////////////////////
@@ -193,12 +181,13 @@ void FTimingProfilerManager::RegisterMajorTabs(IUnrealInsightsModule& InsightsMo
 	if (!bIsTabRegistered && bShouldRegister)
 	{
 		bIsTabRegistered = true;
-		const FInsightsMajorTabConfig& Config = InsightsModule.FindMajorTabConfig(FInsightsManagerTabs::TimingProfilerTabId);
+		const FInsightsMajorTabConfig& Config = InsightsModule.FindMajorTabConfig(MajorTabId);
 		if (Config.bIsAvailable)
 		{
 			// Register tab spawner for the Timing Insights.
-			FTabSpawnerEntry& TabSpawnerEntry = FGlobalTabmanager::Get()->RegisterNomadTabSpawner(FInsightsManagerTabs::TimingProfilerTabId,
-				FOnSpawnTab::CreateRaw(this, &FTimingProfilerManager::SpawnTab), FCanSpawnTab::CreateRaw(this, &FTimingProfilerManager::CanSpawnTab))
+			FTabSpawnerEntry& TabSpawnerEntry = FGlobalTabmanager::Get()->RegisterNomadTabSpawner(MajorTabId,
+					FOnSpawnTab::CreateRaw(this, &FTimingProfilerManager::SpawnTab),
+					FCanSpawnTab::CreateRaw(this, &FTimingProfilerManager::CanSpawnTab))
 				.SetDisplayName(Config.TabLabel.IsSet() ? Config.TabLabel.GetValue() : LOCTEXT("TimingProfilerTabTitle", "Timing Insights"))
 				.SetTooltipText(Config.TabTooltip.IsSet() ? Config.TabTooltip.GetValue() : LOCTEXT("TimingProfilerTooltipText", "Open the Timing Insights tab."))
 				.SetIcon(Config.TabIcon.IsSet() ? Config.TabIcon.GetValue() : FSlateIcon(FInsightsStyle::GetStyleSetName(), "Icons.TimingProfiler"));
@@ -213,7 +202,8 @@ void FTimingProfilerManager::RegisterMajorTabs(IUnrealInsightsModule& InsightsMo
 
 void FTimingProfilerManager::UnregisterMajorTabs()
 {
-	FGlobalTabmanager::Get()->UnregisterNomadTabSpawner(FInsightsManagerTabs::TimingProfilerTabId);
+	FGlobalTabmanager::Get()->UnregisterNomadTabSpawner(MajorTabId);
+	bIsTabRegistered = false;
 }
 
 ////////////////////////////////////////////////////////////////////////////////////////////////////
@@ -239,10 +229,10 @@ TSharedRef<SDockTab> FTimingProfilerManager::SpawnTab(const FSpawnTabArgs& Args)
 
 bool FTimingProfilerManager::CanSpawnTab(const FSpawnTabArgs& Args) const
 {
-#if !WITH_EDITOR
-	return bIsAvailable;
-#else
+#if WITH_EDITOR
 	return true;
+#else
+	return bIsAvailable;
 #endif
 }
 
@@ -251,6 +241,7 @@ bool FTimingProfilerManager::CanSpawnTab(const FSpawnTabArgs& Args) const
 void FTimingProfilerManager::OnTabClosed(TSharedRef<SDockTab> TabBeingClosed)
 {
 	OnWindowClosedEvent();
+
 	RemoveProfilerWindow();
 
 	// Disable TabClosed delegate.
@@ -306,11 +297,10 @@ bool FTimingProfilerManager::Tick(float DeltaTime)
 			MessageLogModule.EnableMessageLogDisplay(true);
 
 #if !WITH_EDITOR
-			const FName& TabId = FInsightsManagerTabs::TimingProfilerTabId;
-			if (FGlobalTabmanager::Get()->HasTabSpawner(TabId))
+			if (FGlobalTabmanager::Get()->HasTabSpawner(MajorTabId))
 			{
-				UE_LOG(LogTimingProfiler, Log, TEXT("Opening the \"Timing Insights\" tab..."));
-				FGlobalTabmanager::Get()->TryInvokeTab(TabId);
+				UE_LOGF(LogTimingProfiler, Log, "Opening the \"Timing Insights\" tab...");
+				FGlobalTabmanager::Get()->TryInvokeTab(MajorTabId);
 			}
 #endif
 		}
@@ -358,10 +348,10 @@ void FTimingProfilerManager::FinishTimerButterflyAggregation()
 void FTimingProfilerManager::ActivateMajorTab()
 {
 	// Ensure Timing Insights / Timing View is the active tab / view.
-	TSharedPtr<SDockTab> TimingInsightsTab = FGlobalTabmanager::Get()->FindExistingLiveTab(FInsightsManagerTabs::TimingProfilerTabId);
-	if (TimingInsightsTab.IsValid())
+	TSharedPtr<SDockTab> MajorTab = FGlobalTabmanager::Get()->FindExistingLiveTab(MajorTabId);
+	if (MajorTab.IsValid())
 	{
-		TimingInsightsTab->ActivateInParent(ETabActivationCause::SetDirectly);
+		MajorTab->ActivateInParent(ETabActivationCause::SetDirectly);
 		ActivateView();
 	}
 }
@@ -372,14 +362,14 @@ void FTimingProfilerManager::ActivateView()
 {
 	// Ensure Timing View is the active view.
 	TSharedPtr<STimingProfilerWindow> Wnd = GetProfilerWindow();
-	if (Wnd.IsValid())
+	if (Wnd)
 	{
 		TSharedPtr<FTabManager> TabManager = Wnd->GetTabManager();
-		TSharedPtr<SDockTab> TimingViewTab = TabManager->FindExistingLiveTab(FTimingProfilerTabs::TimingViewID);
-		if (TimingViewTab.IsValid())
+		TSharedPtr<SDockTab> MinorTab = TabManager->FindExistingLiveTab(FTimingProfilerTabs::TimingViewID);
+		if (MinorTab.IsValid())
 		{
-			TimingViewTab->ActivateInParent(ETabActivationCause::SetDirectly);
-			FSlateApplication::Get().SetKeyboardFocus(TimingViewTab->GetContent());
+			MinorTab->ActivateInParent(ETabActivationCause::SetDirectly);
+			FSlateApplication::Get().SetKeyboardFocus(MinorTab->GetContent());
 		}
 	}
 }
@@ -388,7 +378,7 @@ void FTimingProfilerManager::ActivateView()
 
 void FTimingProfilerManager::OnSessionChanged()
 {
-	UE_LOG(LogTimingProfiler, Log, TEXT("OnSessionChanged"));
+	UE_LOGF(LogTimingProfiler, Log, "OnSessionChanged");
 
 	bIsAvailable = false;
 	if (FInsightsManager::Get()->GetSession().IsValid())
@@ -401,7 +391,7 @@ void FTimingProfilerManager::OnSessionChanged()
 	}
 
 	TSharedPtr<STimingProfilerWindow> Wnd = GetProfilerWindow();
-	if (Wnd.IsValid())
+	if (Wnd)
 	{
 		Wnd->Reset();
 	}
@@ -418,7 +408,7 @@ void FTimingProfilerManager::ShowHideFramesTrack(const bool bIsVisible)
 	bIsFramesTrackVisible = bIsVisible;
 
 	TSharedPtr<STimingProfilerWindow> Wnd = GetProfilerWindow();
-	if (Wnd.IsValid())
+	if (Wnd)
 	{
 		Wnd->ShowHideTab(FTimingProfilerTabs::FramesTrackID, bIsFramesTrackVisible);
 	}
@@ -431,9 +421,9 @@ void FTimingProfilerManager::ShowHideTimingView(const bool bIsVisible)
 	bIsTimingViewVisible = bIsVisible;
 
 	TSharedPtr<STimingProfilerWindow> Wnd = GetProfilerWindow();
-	if (Wnd.IsValid())
+	if (Wnd)
 	{
-		Wnd->ShowHideTab(FTimingProfilerTabs::TimingViewID, bIsTimingViewVisible);
+		Wnd->ShowHideTab(FTimingProfilerTabs::TimingViewID, bIsVisible);
 	}
 }
 
@@ -444,9 +434,9 @@ void FTimingProfilerManager::ShowHideTimersView(const bool bIsVisible)
 	bIsTimersViewVisible = bIsVisible;
 
 	TSharedPtr<STimingProfilerWindow> Wnd = GetProfilerWindow();
-	if (Wnd.IsValid())
+	if (Wnd)
 	{
-		Wnd->ShowHideTab(FTimingProfilerTabs::TimersID, bIsTimersViewVisible);
+		Wnd->ShowHideTab(FTimingProfilerTabs::TimersID, bIsVisible);
 
 		if (bIsTimersViewVisible)
 		{
@@ -462,9 +452,9 @@ void FTimingProfilerManager::ShowHideCallersTreeView(const bool bIsVisible)
 	bIsCallersTreeViewVisible = bIsVisible;
 
 	TSharedPtr<STimingProfilerWindow> Wnd = GetProfilerWindow();
-	if (Wnd.IsValid())
+	if (Wnd)
 	{
-		Wnd->ShowHideTab(FTimingProfilerTabs::CallersID, bIsCallersTreeViewVisible);
+		Wnd->ShowHideTab(FTimingProfilerTabs::CallersID, bIsVisible);
 
 		if (bIsCallersTreeViewVisible)
 		{
@@ -480,9 +470,9 @@ void FTimingProfilerManager::ShowHideCalleesTreeView(const bool bIsVisible)
 	bIsCalleesTreeViewVisible = bIsVisible;
 
 	TSharedPtr<STimingProfilerWindow> Wnd = GetProfilerWindow();
-	if (Wnd.IsValid())
+	if (Wnd)
 	{
-		Wnd->ShowHideTab(FTimingProfilerTabs::CalleesID, bIsCalleesTreeViewVisible);
+		Wnd->ShowHideTab(FTimingProfilerTabs::CalleesID, bIsVisible);
 
 		if (bIsCalleesTreeViewVisible)
 		{
@@ -498,9 +488,9 @@ void FTimingProfilerManager::ShowHideStatsCountersView(const bool bIsVisible)
 	bIsStatsCountersViewVisible = bIsVisible;
 
 	TSharedPtr<STimingProfilerWindow> Wnd = GetProfilerWindow();
-	if (Wnd.IsValid())
+	if (Wnd)
 	{
-		Wnd->ShowHideTab(FTimingProfilerTabs::StatsCountersID, bIsStatsCountersViewVisible);
+		Wnd->ShowHideTab(FTimingProfilerTabs::StatsCountersID, bIsVisible);
 
 		if (bIsStatsCountersViewVisible)
 		{
@@ -516,9 +506,9 @@ void FTimingProfilerManager::ShowHideLogView(const bool bIsVisible)
 	bIsLogViewVisible = bIsVisible;
 
 	TSharedPtr<STimingProfilerWindow> Wnd = GetProfilerWindow();
-	if (Wnd.IsValid())
+	if (Wnd)
 	{
-		Wnd->ShowHideTab(FTimingProfilerTabs::LogViewID, bIsLogViewVisible);
+		Wnd->ShowHideTab(FTimingProfilerTabs::LogViewID, bIsVisible);
 	}
 }
 
@@ -529,9 +519,9 @@ void FTimingProfilerManager::ShowHideModulesView(const bool bIsVisible)
 	bIsModulesViewVisible = bIsVisible;
 
 	TSharedPtr<STimingProfilerWindow> Wnd = GetProfilerWindow();
-	if (Wnd.IsValid())
+	if (Wnd)
 	{
-		Wnd->ShowHideTab(FTimingProfilerTabs::ModulesViewID, bIsModulesViewVisible);
+		Wnd->ShowHideTab(FTimingProfilerTabs::ModulesViewID, bIsVisible);
 	}
 }
 
@@ -732,7 +722,7 @@ void FTimingProfilerManager::OnWindowClosedEvent()
 	if (Wnd)
 	{
 		TSharedPtr<STimingView> TimingView = Wnd->GetTimingView();
-		if (TimingView.IsValid())
+		if (TimingView)
 		{
 			TimingView->CloseQuickFindTab();
 		}
diff --git a/Engine/Source/Developer/TraceInsights/Private/Insights/TimingProfiler/TimingProfilerManager.h b/Engine/Source/Developer/TraceInsights/Private/Insights/TimingProfiler/TimingProfilerManager.h
--- a/Engine/Source/Developer/TraceInsights/Private/Insights/TimingProfiler/TimingProfilerManager.h
+++ b/Engine/Source/Developer/TraceInsights/Private/Insights/TimingProfiler/TimingProfilerManager.h
@@ -70,6 +70,8 @@ public:
 	 */
 	static TSharedPtr<FTimingProfilerManager> Get();
 
+	bool IsAvailable() const { return bIsAvailable; }
+
 	//////////////////////////////////////////////////
 	// IInsightsComponent
 
@@ -100,6 +102,7 @@ public:
 		return ProfilerWindowWeakPtr.Pin();
 	}
 
+	const FName& GetMajorTabId() const { return MajorTabId; }
 	const FName& GetLogListingName() const { return LogListingName; }
 
 	const TSharedRef<FLogFilter>& GetLogFilter() const { return Filter; }
@@ -213,8 +216,12 @@ private:
 	}
 
 private:
-	bool bIsInitialized;
-	bool bIsAvailable;
+	bool bIsInitialized = false;
+	bool bIsAvailable = false;
+
+	/** If the tab is registered in the tab manager. */
+	bool bIsTabRegistered = false;
+
 	FAvailabilityCheck AvailabilityCheck;
 
 	/** The delegate to be invoked when this manager ticks. */
@@ -236,40 +243,41 @@ private:
 	TSharedRef<FLogFilter> Filter;
 
 	/** If the Frames Track is visible or hidden. */
-	bool bIsFramesTrackVisible;
+	bool bIsFramesTrackVisible = false;
 
 	/** If the Timing View is visible or hidden. */
-	bool bIsTimingViewVisible;
+	bool bIsTimingViewVisible = false;
 
 	/** If the Timers View is visible or hidden. */
-	bool bIsTimersViewVisible;
+	bool bIsTimersViewVisible = false;
 
 	/** If the Callers Tree View is visible or hidden. */
-	bool bIsCallersTreeViewVisible;
+	bool bIsCallersTreeViewVisible = false;
 
 	/** If the Callees Tree View is visible or hidden. */
-	bool bIsCalleesTreeViewVisible;
+	bool bIsCalleesTreeViewVisible = false;
 
 	/** If the Stats Counters View is visible or hidden. */
-	bool bIsStatsCountersViewVisible;
+	bool bIsStatsCountersViewVisible = false;
 
 	/** If the Log View is visible or hidden. */
-	bool bIsLogViewVisible;
+	bool bIsLogViewVisible = false;
 
 	/** If the Modules View is visible or hidden. */
-	bool bIsModulesViewVisible;
+	bool bIsModulesViewVisible = false;
 
-	//////////////////////////////////////////////////
-
-	double SelectionStartTime;
-	double SelectionEndTime;
+	double SelectionStartTime = 0.0;
+	double SelectionEndTime = 0.0;
 
 	static constexpr uint32 InvalidTimerId = uint32(-1);
-	uint32 SelectedTimerId;
+	uint32 SelectedTimerId = InvalidTimerId;
 
 	TSharedRef<FTimerButterflyAggregator> TimerButterflyAggregator;
 
-	/** The name of the Timing Insights log listing. */
+	/** The id of the Timing Insights major tab. */
+	FName MajorTabId;
+
+	/** The name of the Timing Profiler log listing. */
 	FName LogListingName;
 
 	ETimingEventsColoringMode ColoringMode = ETimingEventsColoringMode::ByTimerName;
@@ -278,9 +286,6 @@ private:
 
 	/** A shared pointer to the global instance of the Timing Profiler manager. */
 	static TSharedPtr<FTimingProfilerManager> Instance;
-
-	/** If the tab is registered in the tab manager. */
-	bool bIsTabRegistered = false;
 };
 
 } // namespace UE::Insights::TimingProfiler
