Actor matching in Memory View. Initial submit for DetailsView.

#rnx
#changelist validated
#Synced-CL 51753244

diff --git a/Engine/Source/Developer/TraceInsights/Private/Insights/ObjectProfiler/ObjectProfilerManager.cpp b/Engine/Source/Developer/TraceInsights/Private/Insights/ObjectProfiler/ObjectProfilerManager.cpp
--- a/Engine/Source/Developer/TraceInsights/Private/Insights/ObjectProfiler/ObjectProfilerManager.cpp
+++ b/Engine/Source/Developer/TraceInsights/Private/Insights/ObjectProfiler/ObjectProfilerManager.cpp
@@ -14,7 +14,7 @@
 // TraceInsights
 #include "Insights/InsightsStyle.h"
 #include "Insights/ObjectProfiler/ViewModels/ObjectTable.h"
-#include "Insights/ObjectProfiler/Widgets/SObjectTableTreeView.h"
+#include "Insights/ObjectProfiler/Widgets/SMemoryView.h"
 
 #define LOCTEXT_NAMESPACE "UE::Insights::ObjectProfiler"
 
@@ -207,7 +207,7 @@ TSharedRef<SDockTab> FObjectProfilerManager::SpawnTab_ObjectsTableTreeView(const
 		.ShouldAutosize(false)
 		.TabRole(ETabRole::PanelTab)
 		[
-			SAssignNew(ObjectsTableTreeView, SObjectTableTreeView)
+			SAssignNew(MemoryView, SMemoryView)
 		];
 
 	DockTab->SetOnTabClosed(SDockTab::FOnTabClosedCallback::CreateRaw(this, &FObjectProfilerManager::OnObjectsTableTreeViewTabClosed));
@@ -226,21 +226,21 @@ bool FObjectProfilerManager::CanSpawnTab_ObjectsTableTreeView(const FSpawnTabArg
 
 void FObjectProfilerManager::OnObjectsTableTreeViewTabClosed(TSharedRef<SDockTab> TabBeingClosed)
 {
-	if (ObjectsTableTreeView.IsValid())
+	if (MemoryView.IsValid())
 	{
-		ObjectsTableTreeView->OnClose();
+		MemoryView->OnClose();
 	}
 
-	ObjectsTableTreeView.Reset();
+	MemoryView.Reset();
 }
 
 ////////////////////////////////////////////////////////////////////////////////////////////////////
 
 void FObjectProfilerManager::OnWindowClosedEvent()
 {
-	if (ObjectsTableTreeView.IsValid())
+	if (MemoryView.IsValid())
 	{
-		ObjectsTableTreeView->OnClose();
+		MemoryView->OnClose();
 	}
 
 	TSharedPtr<FTabManager> TimingTabManagerSharedPtr = TimingTabManager.Pin();
diff --git a/Engine/Source/Developer/TraceInsights/Private/Insights/ObjectProfiler/ObjectProfilerManager.h b/Engine/Source/Developer/TraceInsights/Private/Insights/ObjectProfiler/ObjectProfilerManager.h
--- a/Engine/Source/Developer/TraceInsights/Private/Insights/ObjectProfiler/ObjectProfilerManager.h
+++ b/Engine/Source/Developer/TraceInsights/Private/Insights/ObjectProfiler/ObjectProfilerManager.h
@@ -21,7 +21,7 @@ namespace UE::Insights::ObjectProfiler
 
 DECLARE_LOG_CATEGORY_EXTERN(LogObjectProfiler, Log, All);
 
-class SObjectTableTreeView;
+class SMemoryView;
 
 struct FObjectProfilerTabs
 {
@@ -96,7 +96,7 @@ private:
 
 	TWeakPtr<FTabManager> TimingTabManager;
 
-	TSharedPtr<SObjectTableTreeView> ObjectsTableTreeView;
+	TSharedPtr<SMemoryView> MemoryView;
 };
 
 } // namespace UE::Insights::ObjectProfiler
diff --git a/Engine/Source/Developer/TraceInsights/Private/Insights/ObjectProfiler/ViewModels/AssetInfoNode.h b/Engine/Source/Developer/TraceInsights/Private/Insights/ObjectProfiler/ViewModels/AssetInfoNode.h
new file mode 100644
--- /dev/null
+++ b/Engine/Source/Developer/TraceInsights/Private/Insights/ObjectProfiler/ViewModels/AssetInfoNode.h
@@ -0,0 +1,25 @@
+// Copyright Epic Games, Inc. All Rights Reserved.
+
+#pragma once
+
+#include "CoreTypes.h"
+
+#include "AssetRegistry/AssetData.h"
+
+namespace UE::Insights::ObjectProfiler
+{
+
+////////////////////////////////////////////////////////////////////////////////////////////////////
+
+/**
+ * Holds asset registry data for a single matched asset entry shown in the object details view.
+ */
+struct FAssetInfoNode
+{
+	FSoftObjectPath ActorPath;
+	FAssetData AssetData;
+};
+
+////////////////////////////////////////////////////////////////////////////////////////////////////
+
+} // namespace UE::Insights::ObjectProfiler
diff --git a/Engine/Source/Developer/TraceInsights/Private/Insights/ObjectProfiler/Widgets/SMemoryView.cpp b/Engine/Source/Developer/TraceInsights/Private/Insights/ObjectProfiler/Widgets/SMemoryView.cpp
new file mode 100644
--- /dev/null
+++ b/Engine/Source/Developer/TraceInsights/Private/Insights/ObjectProfiler/Widgets/SMemoryView.cpp
@@ -0,0 +1,68 @@
+// Copyright Epic Games, Inc. All Rights Reserved.
+
+#include "SMemoryView.h"
+
+#include "Widgets/Layout/SSplitter.h"
+
+// TraceInsights
+#include "Insights/ObjectProfiler/Widgets/SObjectDetailsView.h"
+#include "Insights/ObjectProfiler/Widgets/SObjectTableTreeView.h"
+
+#define LOCTEXT_NAMESPACE "UE::Insights::ObjectProfiler::SMemoryView"
+
+namespace UE::Insights::ObjectProfiler
+{
+
+////////////////////////////////////////////////////////////////////////////////////////////////////
+
+SMemoryView::SMemoryView()
+{
+}
+
+////////////////////////////////////////////////////////////////////////////////////////////////////
+
+SMemoryView::~SMemoryView()
+{
+}
+
+////////////////////////////////////////////////////////////////////////////////////////////////////
+
+void SMemoryView::Construct(const FArguments& InArgs)
+{
+	SAssignNew(ObjectDetailsView, SObjectDetailsView, SharedThis(this));
+	SAssignNew(ObjectTableTreeView, SObjectTableTreeView, ObjectDetailsView);
+
+	ChildSlot
+	[
+		SNew(SSplitter)
+		.Orientation(Orient_Horizontal)
+
+		+ SSplitter::Slot()
+		.Value(0.6f)
+		[
+			ObjectTableTreeView.ToSharedRef()
+		]
+
+		+ SSplitter::Slot()
+		.Value(0.4f)
+		[
+			ObjectDetailsView.ToSharedRef()
+		]
+	];
+}
+
+////////////////////////////////////////////////////////////////////////////////////////////////////
+
+void SMemoryView::OnClose()
+{
+	if (ObjectTableTreeView.IsValid())
+	{
+		ObjectTableTreeView->OnClose();
+	}
+}
+
+////////////////////////////////////////////////////////////////////////////////////////////////////
+
+} // namespace UE::Insights::ObjectProfiler
+
+#undef LOCTEXT_NAMESPACE
diff --git a/Engine/Source/Developer/TraceInsights/Private/Insights/ObjectProfiler/Widgets/SMemoryView.h b/Engine/Source/Developer/TraceInsights/Private/Insights/ObjectProfiler/Widgets/SMemoryView.h
new file mode 100644
--- /dev/null
+++ b/Engine/Source/Developer/TraceInsights/Private/Insights/ObjectProfiler/Widgets/SMemoryView.h
@@ -0,0 +1,61 @@
+// Copyright Epic Games, Inc. All Rights Reserved.
+
+#pragma once
+
+#include "CoreTypes.h"
+
+#include "Templates/SharedPointer.h"
+#include "Widgets/DeclarativeSyntaxSupport.h"
+#include "Widgets/SCompoundWidget.h"
+
+#include "Insights/ObjectProfiler/ViewModels/AssetInfoNode.h"
+#include "Insights/Widgets/IObjectProfilerMemoryView.h"
+
+namespace UE::Insights::ObjectProfiler
+{
+
+class SObjectDetailsView;
+class SObjectTableTreeView;
+
+////////////////////////////////////////////////////////////////////////////////////////////////////
+
+/**
+ * A two-panel widget containing an object tree view on the left and an object details view on the right.
+ */
+class SMemoryView : public SCompoundWidget, public IObjectProfilerMemoryView
+{
+public:
+	/** Default constructor. */
+	SMemoryView();
+
+	/** Virtual destructor. */
+	virtual ~SMemoryView();
+
+	SLATE_BEGIN_ARGS(SMemoryView) {}
+	SLATE_END_ARGS()
+
+	/**
+	 * Construct this widget.
+	 * @param InArgs - The declaration data for this widget
+	 */
+	void Construct(const FArguments& InArgs);
+
+	/** Notifies child widgets that the view is closing. */
+	void OnClose();
+
+	virtual TSharedRef<SWidget> GetWidget() override { return SharedThis(this); }
+
+	virtual void SetGetActorsDelegate(FGetActors InDelegate) { GetActorsDelegate = MoveTemp(InDelegate); };
+	FGetActors GetActorsDelegate;
+
+private:
+	/** The object tree view displayed in the left panel. */
+	TSharedPtr<SObjectTableTreeView> ObjectTableTreeView;
+
+	/** The object details view displayed in the right panel. */
+	TSharedPtr<SObjectDetailsView> ObjectDetailsView;
+};
+
+////////////////////////////////////////////////////////////////////////////////////////////////////
+
+} // namespace UE::Insights::ObjectProfiler
diff --git a/Engine/Source/Developer/TraceInsights/Private/Insights/ObjectProfiler/Widgets/SObjectDetailsView.cpp b/Engine/Source/Developer/TraceInsights/Private/Insights/ObjectProfiler/Widgets/SObjectDetailsView.cpp
new file mode 100644
--- /dev/null
+++ b/Engine/Source/Developer/TraceInsights/Private/Insights/ObjectProfiler/Widgets/SObjectDetailsView.cpp
@@ -0,0 +1,291 @@
+// Copyright Epic Games, Inc. All Rights Reserved.
+
+#include "SObjectDetailsView.h"
+
+#include "Brushes/SlateNoResource.h"
+#include "Modules/ModuleManager.h"
+#if WITH_EDITOR
+#include "ContentBrowserModule.h"
+#include "IContentBrowserSingleton.h"
+#include "UnrealEdGlobals.h"
+#include "Editor/EditorEngine.h"
+#include "Editor/UnrealEdEngine.h"
+
+extern UNREALED_API class UEditorEngine* GEditor;
+
+#endif
+#include "SlateOptMacros.h"
+#include "Styling/AppStyle.h"
+#include "Styling/SlateTypes.h"
+#include "Widgets/Images/SImage.h"
+#include "Widgets/Input/SButton.h"
+#include "Widgets/Layout/SScrollBar.h"
+#include "Widgets/SBoxPanel.h"
+#include "Widgets/Text/STextBlock.h"
+#include "Widgets/Views/SListView.h"
+#include "Widgets/Views/STableRow.h"
+
+// TraceInsights
+#include "Insights/ObjectProfiler/ViewModels/AssetInfoNode.h"
+#include "SMemoryView.h"
+
+#define LOCTEXT_NAMESPACE "UE::Insights::ObjectProfiler::SObjectDetailsView"
+
+namespace UE::Insights::ObjectProfiler
+{
+
+////////////////////////////////////////////////////////////////////////////////////////////////////
+// SObjectDetailsViewRow
+////////////////////////////////////////////////////////////////////////////////////////////////////
+
+class SObjectDetailsViewRow : public STableRow<TSharedPtr<FAssetInfoNode>>
+{
+	SLATE_BEGIN_ARGS(SObjectDetailsViewRow) {}
+	SLATE_END_ARGS()
+
+public:
+	void Construct(const FArguments& InArgs, TSharedPtr<FAssetInfoNode> InItem, const TSharedRef<STableViewBase>& InOwnerTableView)
+	{
+		static const FSlateNoResource NoBrush;
+		static const FTableRowStyle NoBackgroundRowStyle = FTableRowStyle()
+			.SetEvenRowBackgroundBrush(NoBrush)
+			.SetEvenRowBackgroundHoveredBrush(NoBrush)
+			.SetOddRowBackgroundBrush(NoBrush)
+			.SetOddRowBackgroundHoveredBrush(NoBrush)
+			.SetSelectorFocusedBrush(NoBrush)
+			.SetActiveBrush(NoBrush)
+			.SetActiveHoveredBrush(NoBrush)
+			.SetInactiveBrush(NoBrush)
+			.SetInactiveHoveredBrush(NoBrush);
+
+		WeakItem = MoveTemp(InItem);
+		STableRow<TSharedPtr<FAssetInfoNode>>::Construct(
+			STableRow<TSharedPtr<FAssetInfoNode>>::FArguments()
+			.Style(&NoBackgroundRowStyle)
+			.Content()
+			[
+				SNew(SBox)
+				.Padding(FMargin(4.0f, 3.0f))
+				[
+					SNew(SHorizontalBox)
+
+					+ SHorizontalBox::Slot()
+					.FillWidth(1.0f)
+					.VAlign(VAlign_Center)
+					[
+						SNew(STextBlock)
+						.Text(this, &SObjectDetailsViewRow::GetAssetNameText)
+						.ToolTipText(this, &SObjectDetailsViewRow::GetAssetTooltipText)
+					]
+
+					//+ SHorizontalBox::Slot()
+					//.AutoWidth()
+					//.VAlign(VAlign_Center)
+					//.Padding(4.0f, 0.0f, 0.0f, 0.0f)
+					//[
+					//	SNew(SButton)
+					//	.ButtonStyle(FAppStyle::Get(), "SimpleButton")
+					//	.ToolTipText(LOCTEXT("BrowseToAssetTooltip", "Browse to asset in Content Browser"))
+					//	.OnClicked(this, &SObjectDetailsViewRow::OnBrowseToAssetClicked)
+					//	.IsEnabled(this, &SObjectDetailsViewRow::IsBrowseToAssetEnabled)
+					//	[
+					//		SNew(SImage)
+					//		.Image(FAppStyle::Get().GetBrush("Icons.BrowseContent"))
+					//	]
+					//]
+
+					+ SHorizontalBox::Slot()
+					.AutoWidth()
+					.VAlign(VAlign_Center)
+					.Padding(4.0f, 0.0f, 0.0f, 0.0f)
+					[
+						SNew(SButton)
+						.ButtonStyle(FAppStyle::Get(), "SimpleButton")
+						.ToolTipText(LOCTEXT("BrowseToActorTooltip", "Browse to Actor"))
+						.OnClicked(this, &SObjectDetailsViewRow::OnBrowseToActorClicked)
+						.IsEnabled(this, &SObjectDetailsViewRow::IsBrowseToActorEnabled)
+						[
+							SNew(SImage)
+							.Image(FAppStyle::Get().GetBrush("Icons.BrowseContent"))
+						]
+					]
+				]
+			],
+			InOwnerTableView);
+	}
+
+private:
+	FText GetAssetNameText() const
+	{
+		TSharedPtr<FAssetInfoNode> Item = WeakItem.Pin();
+		return Item.IsValid() ? FText::FromString(Item->ActorPath.GetSubPathString()) : FText::GetEmpty();
+	}
+
+	FText GetAssetTooltipText() const
+	{
+		TSharedPtr<FAssetInfoNode> Item = WeakItem.Pin();
+		return Item.IsValid() ? FText::FromString(Item->ActorPath.GetAssetPathString()) : FText::GetEmpty();
+	}
+
+	bool IsBrowseToAssetEnabled() const
+	{
+		TSharedPtr<FAssetInfoNode> Item = WeakItem.Pin();
+		return Item.IsValid() && Item->AssetData.IsValid();
+	}
+
+	FReply OnBrowseToAssetClicked()
+	{
+#if WITH_EDITOR
+		TSharedPtr<FAssetInfoNode> Item = WeakItem.Pin();
+		if (Item.IsValid() && Item->AssetData.IsValid())
+		{
+			if (FContentBrowserModule* ContentBrowserModule = FModuleManager::GetModulePtr<FContentBrowserModule>("ContentBrowser"))
+			{
+				ContentBrowserModule->Get().SyncBrowserToAssets({ Item->AssetData });
+			}
+		}
+#endif
+		return FReply::Handled();
+	}
+
+	bool IsBrowseToActorEnabled() const
+	{
+		TSharedPtr<FAssetInfoNode> Item = WeakItem.Pin();
+		return Item.IsValid() && Item->ActorPath.IsValid();
+	}
+
+	FReply OnBrowseToActorClicked()
+	{
+#if WITH_EDITOR
+		TSharedPtr<FAssetInfoNode> Item = WeakItem.Pin();
+		if (Item.IsValid() && Item->ActorPath.IsValid())
+		{
+			UObject* Obj = Item->ActorPath.ResolveObject();
+			if (!Obj)
+			{
+				Obj = Item->ActorPath.TryLoad();
+			}
+			if (AActor* AsActor = Cast<AActor>(Obj))
+			{
+				GEditor->SelectNone(false, true);
+				GEditor->SelectActor(AsActor, /*InSelected=*/true, /*bNotify=*/false, /*bSelectEvenIfHidden=*/true);
+				GEditor->NoteSelectionChange();
+				GEditor->MoveViewportCamerasToActor(*AsActor, false);
+				GUnrealEd->ShowActorProperties();
+				GUnrealEd->UpdateFloatingPropertyWindows();
+			}
+		}
+#endif
+		return FReply::Handled();
+	}
+
+private:
+	TWeakPtr<FAssetInfoNode> WeakItem;
+};
+
+////////////////////////////////////////////////////////////////////////////////////////////////////
+// SObjectDetailsView
+////////////////////////////////////////////////////////////////////////////////////////////////////
+
+SObjectDetailsView::SObjectDetailsView()
+{
+}
+
+////////////////////////////////////////////////////////////////////////////////////////////////////
+
+SObjectDetailsView::~SObjectDetailsView()
+{
+}
+
+////////////////////////////////////////////////////////////////////////////////////////////////////
+
+BEGIN_SLATE_FUNCTION_BUILD_OPTIMIZATION
+void SObjectDetailsView::Construct(const FArguments& InArgs, TWeakPtr<SMemoryView> MemoryView)
+{
+	ParentMemoryView = MemoryView;
+
+	SAssignNew(ExternalScrollbar, SScrollBar)
+	.AlwaysShowScrollbar(true);
+
+	ChildSlot
+	[
+		SNew(SVerticalBox)
+
+		+ SVerticalBox::Slot()
+		.AutoHeight()
+		.HAlign(HAlign_Center)
+		.Padding(FMargin(0.0f, 4.0f))
+		[
+			SNew(STextBlock)
+			.Text(LOCTEXT("ObjectDetailsTitle", "Object Details"))
+		]
+
+		+ SVerticalBox::Slot()
+		.FillHeight(1.0f)
+		[
+			SNew(SHorizontalBox)
+
+			+ SHorizontalBox::Slot()
+			.FillWidth(1.0f)
+			.Padding(0.0f)
+			.VAlign(VAlign_Fill)
+			[
+				SAssignNew(ListView, SListView<TSharedPtr<FAssetInfoNode>>)
+				.ExternalScrollbar(ExternalScrollbar)
+				.SelectionMode(ESelectionMode::Single)
+				.ListItemsSource(&AssetEntries)
+				.OnGenerateRow(this, &SObjectDetailsView::OnGenerateRow)
+			]
+
+			+ SHorizontalBox::Slot()
+			.AutoWidth()
+			.Padding(0.0f)
+			[
+				ExternalScrollbar.ToSharedRef()
+			]
+		]
+	];
+}
+END_SLATE_FUNCTION_BUILD_OPTIMIZATION
+
+////////////////////////////////////////////////////////////////////////////////////////////////////
+
+void SObjectDetailsView::SetMatchedAssets(const FAssetData& AssetData)
+{
+	AssetEntries.Empty();
+
+	if (AssetData.IsValid())
+	{
+		auto MemoryView = ParentMemoryView.Pin();
+		if (MemoryView)
+		{
+			if (MemoryView->GetActorsDelegate.IsBound())
+			{
+				TArray<FSoftObjectPath> Actors;
+				MemoryView->GetActorsDelegate.Execute(AssetData, Actors);
+
+				for (auto Actor : Actors)
+				{
+					auto Entry = MakeShared<FAssetInfoNode>();
+					Entry->ActorPath = Actor;
+					AssetEntries.Add(Entry);
+				}
+			}
+		}
+	}
+
+	ListView->RequestListRefresh();
+}
+
+////////////////////////////////////////////////////////////////////////////////////////////////////
+
+TSharedRef<ITableRow> SObjectDetailsView::OnGenerateRow(TSharedPtr<FAssetInfoNode> InItem, const TSharedRef<STableViewBase>& OwnerTable)
+{
+	return SNew(SObjectDetailsViewRow, InItem, OwnerTable);
+}
+
+////////////////////////////////////////////////////////////////////////////////////////////////////
+
+} // namespace UE::Insights::ObjectProfiler
+
+#undef LOCTEXT_NAMESPACE
diff --git a/Engine/Source/Developer/TraceInsights/Private/Insights/ObjectProfiler/Widgets/SObjectDetailsView.h b/Engine/Source/Developer/TraceInsights/Private/Insights/ObjectProfiler/Widgets/SObjectDetailsView.h
new file mode 100644
--- /dev/null
+++ b/Engine/Source/Developer/TraceInsights/Private/Insights/ObjectProfiler/Widgets/SObjectDetailsView.h
@@ -0,0 +1,53 @@
+// Copyright Epic Games, Inc. All Rights Reserved.
+
+#pragma once
+
+#include "CoreTypes.h"
+
+#include "Widgets/DeclarativeSyntaxSupport.h"
+#include "Widgets/SCompoundWidget.h"
+#include "Widgets/Views/SListView.h"
+
+#include "Insights/ObjectProfiler/ViewModels/AssetInfoNode.h"
+
+namespace UE::Insights::ObjectProfiler
+{
+class SMemoryView;
+////////////////////////////////////////////////////////////////////////////////////////////////////
+
+/**
+ * A widget that displays a list of matched asset names for the current object selection.
+ */
+class SObjectDetailsView : public SCompoundWidget
+{
+public:
+	/** Default constructor. */
+	SObjectDetailsView();
+
+	/** Virtual destructor. */
+	virtual ~SObjectDetailsView();
+
+	SLATE_BEGIN_ARGS(SObjectDetailsView) {}
+	SLATE_END_ARGS()
+
+	void Construct(const FArguments& InArgs, TWeakPtr<SMemoryView> MemoryView);
+
+	void SetMatchedAssets(const FAssetData& AssetData);
+
+protected:
+	TSharedRef<ITableRow> OnGenerateRow(TSharedPtr<FAssetInfoNode> InItem, const TSharedRef<STableViewBase>& OwnerTable);
+
+private:
+	TSharedPtr<SListView<TSharedPtr<FAssetInfoNode>>> ListView;
+
+	TSharedPtr<SScrollBar> ExternalScrollbar;
+
+	/** List of matched asset info nodes shown in the list view. */
+	TArray<TSharedPtr<FAssetInfoNode>> AssetEntries;
+
+	TWeakPtr<SMemoryView> ParentMemoryView;
+};
+
+////////////////////////////////////////////////////////////////////////////////////////////////////
+
+} // namespace UE::Insights::ObjectProfiler
diff --git a/Engine/Source/Developer/TraceInsights/Private/Insights/ObjectProfiler/Widgets/SObjectTableTreeView.cpp b/Engine/Source/Developer/TraceInsights/Private/Insights/ObjectProfiler/Widgets/SObjectTableTreeView.cpp
--- a/Engine/Source/Developer/TraceInsights/Private/Insights/ObjectProfiler/Widgets/SObjectTableTreeView.cpp
+++ b/Engine/Source/Developer/TraceInsights/Private/Insights/ObjectProfiler/Widgets/SObjectTableTreeView.cpp
@@ -2,6 +2,9 @@
 
 #include "SObjectTableTreeView.h"
 
+#include "AssetRegistry/AssetRegistryModule.h"
+#include "AssetRegistry/IAssetRegistry.h"
+#include "AssetRegistry/AssetData.h"
 #include "DesktopPlatformModule.h"
 #include "Framework/MultiBox/MultiBoxBuilder.h"
 #include "Misc/ConfigCacheIni.h"
@@ -30,6 +33,7 @@
 #include "Insights/ObjectProfiler/ViewModels/ObjectNode.h"
 #include "Insights/ObjectProfiler/ViewModels/ObjectTable.h"
 #include "Insights/ObjectProfiler/ViewModels/ObjectTableViewPresets.h"
+#include "Insights/ObjectProfiler/Widgets/SObjectDetailsView.h"
 
 #define LOCTEXT_NAMESPACE "UE::Insights::ObjectProfiler::SObjectTableTreeView"
 
@@ -51,8 +55,10 @@ SObjectTableTreeView::~SObjectTableTreeView()
 
 ////////////////////////////////////////////////////////////////////////////////////////////////////
 
-void SObjectTableTreeView::Construct(const FArguments& InArgs)
+void SObjectTableTreeView::Construct(const FArguments& InArgs, TWeakPtr<SObjectDetailsView> InDetailsView)
 {
+	DetailsView = InDetailsView;
+
 	TSharedRef<FObjectTable> TablePtr = MakeShared<FObjectTable>();
 	TablePtr->Reset();
 	TablePtr->SetDisplayName(LOCTEXT("ObjectsTableName", "Objects"));
@@ -246,6 +252,60 @@ END_SLATE_FUNCTION_BUILD_OPTIMIZATION
 
 ////////////////////////////////////////////////////////////////////////////////////////////////////
 
+void SObjectTableTreeView::FindMatchedAsset(const TArray<FTableTreeNodePtr>& SelectedNodes, FAssetData& OutAsset) const
+{
+	FAssetRegistryModule& AssetRegistryModule = FModuleManager::Get().LoadModuleChecked<FAssetRegistryModule>(TEXT("AssetRegistry"));
+
+	for (const FTableTreeNodePtr& SelectedNode : SelectedNodes)
+	{
+		if (SelectedNode->Is<FObjectNode>())
+		{
+			const FObjectNode& ObjectNode = SelectedNode->As<FObjectNode>();
+
+			FString Path;
+			auto Package = ObjectNode.GetPackage();
+			if (Package)
+			{
+				Path = Package->GetObjectName();
+			}
+			else
+			{
+				Path = ObjectNode.GetObjectName();
+			}
+
+			auto Outer = ObjectNode.GetOuter();
+			bool Result = false;
+			do
+			{
+				OutAsset = AssetRegistryModule.Get().GetAssetByObjectPath(Path, false, false);
+				if (OutAsset.IsValid())
+				{
+					return;
+				}
+				if (Outer.IsValid())
+				{
+					auto PackageOuter = Outer->GetPackage();
+					if (PackageOuter)
+					{
+						Path = PackageOuter->GetObjectPath();
+					}
+					else
+					{
+						Path = Outer->GetObjectName();
+					}
+					Outer = Outer->GetOuter();
+				}
+				else
+				{
+					break;
+				}
+			} while (!OutAsset.IsValid());
+		}
+	}
+}
+
+////////////////////////////////////////////////////////////////////////////////////////////////////
+
 void SObjectTableTreeView::UpdateSelectionStatsText()
 {
 	TArray<FTableTreeNodePtr> SelectedNodes;
@@ -293,6 +353,18 @@ void SObjectTableTreeView::UpdateSelectionStatsText()
 void SObjectTableTreeView::TreeView_OnSelectionChanged(FTableTreeNodePtr SelectedItem, ESelectInfo::Type SelectInfo)
 {
 	UpdateSelectionStatsText();
+
+	TArray<FTableTreeNodePtr> SelectedNodes;
+	TreeView->GetSelectedItems(SelectedNodes);
+
+	FAssetData Asset;
+	FindMatchedAsset(SelectedNodes, Asset);
+
+	TSharedPtr<SObjectDetailsView> DetailsViewShared = DetailsView.Pin();
+	if (DetailsViewShared)
+	{
+		DetailsViewShared->SetMatchedAssets(Asset);
+	}
 }
 
 ////////////////////////////////////////////////////////////////////////////////////////////////////
diff --git a/Engine/Source/Developer/TraceInsights/Private/Insights/ObjectProfiler/Widgets/SObjectTableTreeView.h b/Engine/Source/Developer/TraceInsights/Private/Insights/ObjectProfiler/Widgets/SObjectTableTreeView.h
--- a/Engine/Source/Developer/TraceInsights/Private/Insights/ObjectProfiler/Widgets/SObjectTableTreeView.h
+++ b/Engine/Source/Developer/TraceInsights/Private/Insights/ObjectProfiler/Widgets/SObjectTableTreeView.h
@@ -4,7 +4,10 @@
 
 #include "CoreTypes.h"
 
+#include "Templates/SharedPointer.h"
+
 // TraceInsights
+#include "Insights/ObjectProfiler/ViewModels/AssetInfoNode.h"
 #include "Insights/Table/Widgets/SSessionTableTreeView.h"
 
 namespace UE::Insights::ObjectProfiler
@@ -12,6 +15,10 @@ namespace UE::Insights::ObjectProfiler
 
 class FObjectNode;
 class FObjectTable;
+class SObjectDetailsView;
+
+DECLARE_DELEGATE_OneParam(FOnMatchedAssetsChanged, const TArray<TSharedPtr<FAssetInfoNode>>&);
+DECLARE_DELEGATE_RetVal_OneParam(bool, FIsAssetAvailable, const FAssetData& Asset);
 
 ////////////////////////////////////////////////////////////////////////////////////////////////////
 /**
@@ -33,7 +40,7 @@ public:
 	 * Construct this widget
 	 * @param InArgs - The declaration data for this widget
 	 */
-	void Construct(const FArguments& InArgs);
+	void Construct(const FArguments& InArgs, TWeakPtr<SObjectDetailsView> DetailsView);
 
 	virtual void ConstructHeaderArea(TSharedRef<SVerticalBox> InHostBox) override;
 	virtual void ConstructFooterArea(TSharedRef<SVerticalBox> InHostBox) override;
@@ -58,6 +65,9 @@ public:
 	 */
 	virtual void RebuildTree(bool bResync);
 
+	FOnMatchedAssetsChanged& GetOnMatchedAssetsChangedEvent() { return OnMatchedAssetsChangedEvent; }
+	FIsAssetAvailable& GetIsAssetAvailableEvent() { return IsAssetAvailableEvent; }
+
 protected:
 	virtual void InternalCreateGroupings() override;
 
@@ -73,6 +83,7 @@ protected:
 
 private:
 	void UpdateSelectionStatsText();
+	void FindMatchedAsset(const TArray<FTableTreeNodePtr>& SelectedNodes, FAssetData& OutAsset) const;
 	FText GetNumSelectedObjectsText() const { return NumSelectedObjectsText; }
 	FText GetSelectedObjectsText() const { return  SelectedObjectsText; }
 	FText GetSelectionStructSizeText() const { return SelectionStructSizeText; }
@@ -108,6 +119,11 @@ private:
 	FText SelectionStructSizeText;
 	FText SelectionSystemMemSizeText;
 	FText SelectionVideoMemSizeText;
+
+	FOnMatchedAssetsChanged OnMatchedAssetsChangedEvent;
+	FIsAssetAvailable IsAssetAvailableEvent;
+
+	TWeakPtr<SObjectDetailsView> DetailsView;
 };
 
 ////////////////////////////////////////////////////////////////////////////////////////////////////
diff --git a/Engine/Source/Developer/TraceInsights/Private/Insights/TraceInsightsModule.cpp b/Engine/Source/Developer/TraceInsights/Private/Insights/TraceInsightsModule.cpp
--- a/Engine/Source/Developer/TraceInsights/Private/Insights/TraceInsightsModule.cpp
+++ b/Engine/Source/Developer/TraceInsights/Private/Insights/TraceInsightsModule.cpp
@@ -36,7 +36,7 @@
 #include "Insights/MemoryProfiler/MemoryProfilerManager.h"
 #include "Insights/NetworkingProfiler/NetworkingProfilerManager.h"
 #include "Insights/ObjectProfiler/ObjectProfilerManager.h"
-#include "Insights/ObjectProfiler/Widgets/SObjectTableTreeView.h"
+#include "Insights/ObjectProfiler/Widgets/SMemoryView.h"
 #include "Insights/TaskGraphProfiler/TaskGraphProfilerManager.h"
 #include "Insights/Tests/InsightsTestRunner.h"
 #include "Insights/TimingProfiler/TimingProfilerManager.h"
@@ -625,9 +625,9 @@ void FTraceInsightsModule::HandleCodeAccessorOpenFileFailed(const FString& Filen
 
 ////////////////////////////////////////////////////////////////////////////////////////////////////
 
-TSharedRef<SWidget> FTraceInsightsModule::CreateObjectTableTreeView() const
+TSharedRef<UE::Insights::ObjectProfiler::IObjectProfilerMemoryView> FTraceInsightsModule::CreateMemoryView() const
 {
-	return SNew(UE::Insights::ObjectProfiler::SObjectTableTreeView);
+	return SNew(UE::Insights::ObjectProfiler::SMemoryView);
 }
 
 ////////////////////////////////////////////////////////////////////////////////////////////////////
diff --git a/Engine/Source/Developer/TraceInsights/Private/Insights/TraceInsightsModule.h b/Engine/Source/Developer/TraceInsights/Private/Insights/TraceInsightsModule.h
--- a/Engine/Source/Developer/TraceInsights/Private/Insights/TraceInsightsModule.h
+++ b/Engine/Source/Developer/TraceInsights/Private/Insights/TraceInsightsModule.h
@@ -92,7 +92,7 @@ public:
 	virtual void ScheduleCommand(const FString& InCmd) override;
 	virtual bool Exec(const TCHAR* Cmd, FOutputDevice& Ar) override;
 
-	virtual TSharedRef<SWidget> CreateObjectTableTreeView() const override;
+	virtual TSharedRef<UE::Insights::ObjectProfiler::IObjectProfilerMemoryView> CreateMemoryView() const override;
 
 private:
 	void RegisterTabSpawners();
diff --git a/Engine/Source/Developer/TraceInsights/Public/Insights/IUnrealInsightsModule.h b/Engine/Source/Developer/TraceInsights/Public/Insights/IUnrealInsightsModule.h
--- a/Engine/Source/Developer/TraceInsights/Public/Insights/IUnrealInsightsModule.h
+++ b/Engine/Source/Developer/TraceInsights/Public/Insights/IUnrealInsightsModule.h
@@ -16,6 +16,9 @@
 // TraceAnalysis
 #include "Trace/DataStream.h"
 
+// TraceInsights
+#include "Insights/Widgets/IObjectProfilerMemoryView.h"
+
 #define UE_API TRACEINSIGHTS_API
 
 class FExtender;
@@ -385,9 +388,9 @@ public:
 	virtual bool Exec(const TCHAR* Cmd, FOutputDevice& Ar) = 0;
 
 	/**
-	* Creates and returns an ObjectTableTreeView.
+	* Creates and returns an ObjectProfilerMemoryView object.
 	*/
-	virtual TSharedRef<SWidget> CreateObjectTableTreeView() const = 0;
+	virtual TSharedRef<UE::Insights::ObjectProfiler::IObjectProfilerMemoryView> CreateMemoryView() const = 0;
 };
 
 ////////////////////////////////////////////////////////////////////////////////////////////////////
diff --git a/Engine/Source/Developer/TraceInsights/Public/Insights/Widgets/IObjectProfilerMemoryView.h b/Engine/Source/Developer/TraceInsights/Public/Insights/Widgets/IObjectProfilerMemoryView.h
new file mode 100644
--- /dev/null
+++ b/Engine/Source/Developer/TraceInsights/Public/Insights/Widgets/IObjectProfilerMemoryView.h
@@ -0,0 +1,36 @@
+// Copyright Epic Games, Inc. All Rights Reserved.
+
+#pragma once
+
+#include "AssetRegistry/AssetData.h"
+#include "Containers/StringFwd.h"
+#include "Widgets/DeclarativeSyntaxSupport.h"
+#include "Widgets/SCompoundWidget.h"
+
+#define UE_API TRACEINSIGHTS_API
+
+////////////////////////////////////////////////////////////////////////////////////////////////////
+
+namespace UE::Insights::ObjectProfiler
+{
+
+/** Interface for the Object Profiler Memory View. */
+class IObjectProfilerMemoryView
+{
+public:
+
+	/** Default constructor. */
+	IObjectProfilerMemoryView() {}
+
+	/** Virtual destructor. */
+	 virtual ~IObjectProfilerMemoryView() {}
+
+	virtual TSharedRef<SWidget> GetWidget() = 0;
+	
+	DECLARE_DELEGATE_RetVal_TwoParams(bool, FGetActors, const FAssetData& Asset, TArray<FSoftObjectPath>& OutActors);
+	virtual void SetGetActorsDelegate(FGetActors InDelegate) = 0;
+};
+
+} // namespace UE::Insights::ObjectProfiler
+
+#undef UE_API
diff --git a/Engine/Source/Developer/TraceInsights/TraceInsights.Build.cs b/Engine/Source/Developer/TraceInsights/TraceInsights.Build.cs
--- a/Engine/Source/Developer/TraceInsights/TraceInsights.Build.cs
+++ b/Engine/Source/Developer/TraceInsights/TraceInsights.Build.cs
@@ -12,6 +12,7 @@ public class TraceInsights : ModuleRules
 			new string[] {
 				"AppFramework", // for SColorPicker, RestoreStarshipSuite()
 				"ApplicationCore",
+				"AssetRegistry",
 				"AutomationDriver",
 				"Cbor",
 				"Core",
@@ -51,7 +52,9 @@ public class TraceInsights : ModuleRules
 		{
 			PrivateDependencyModuleNames.AddRange(
 				new string[] {
+					"ContentBrowser",
 					"Engine",
+					"UnrealEd",
 				}
 			);
 		}
