[Insights] Added dependency to AssetDefiniton (editor only).

#rnx
#Synced-CL 52097700

diff --git a/Engine/Source/Developer/TraceInsights/TraceInsights.Build.cs b/Engine/Source/Developer/TraceInsights/TraceInsights.Build.cs
--- a/Engine/Source/Developer/TraceInsights/TraceInsights.Build.cs
+++ b/Engine/Source/Developer/TraceInsights/TraceInsights.Build.cs
@@ -53,6 +53,7 @@ public class TraceInsights : ModuleRules
 			PrivateDependencyModuleNames.AddRange(
 				new string[] {
 					"ContentBrowser",
+					"AssetDefinition",
 					"Engine",
 					"UnrealEd",
 				}
