[Insights] ObjectProfiler: Added initial interface for IAssetInfoProvider with two implementations: FAssetInfoProvider and FAssetToolsInfoProvider.

#Synced-CL 52097700

diff --git a/Engine/Source/Developer/TraceInsights/Private/Insights/ObjectProfiler/AssetInfoProvider.cpp b/Engine/Source/Developer/TraceInsights/Private/Insights/ObjectProfiler/AssetInfoProvider.cpp
new file mode 100644
--- /dev/null
+++ b/Engine/Source/Developer/TraceInsights/Private/Insights/ObjectProfiler/AssetInfoProvider.cpp
@@ -0,0 +1,157 @@
+// Copyright Epic Games, Inc. All Rights Reserved.
+
+#include "AssetInfoProvider.h"
+
+#if UE_INSIGHTS_WITH_CUSTOM_ASSET_INFO
+
+#define LOCTEXT_NAMESPACE "UE::Insights::ObjectProfiler::AssetInfoProvider"
+
+#define UE_INSIGHTS_ENUM_ASSET_CATEGORIES(macro) \
+	macro(AI,           AI,                              FColor::White,          "Artificial Intelligence") \
+	macro(Animation,    Animation,                       FColor(80, 123, 72),    "Animation") \
+	macro(Animation,    AnimBoneCompressionSettings,     FColor(255, 255, 0),    "Animation Bone Compression Settings") \
+	macro(Animation,    AnimCurveCompressionSettings,    FColor(255, 255, 0),    "Animation Curve Compression Settings") \
+	macro(Animation,    VariableFrameStrippingSettings,  FColor(255, 255, 0),    "Variable Frame Stripping Settings") \
+	macro(Audio,        Audio,                           FColor(128, 255, 128),  "Audio") \
+	macro(Audio,        MetaSoundSource,                 FColor(64, 255, 64),    "MetaSound Source") \
+	macro(Audio,        SoundWave,                       FColor(255, 255, 0),    "Sound Wave") \
+	macro(Audio,        SoundAttenuation,                FColor(75, 114, 228),   "Sound Attenuation") \
+	macro(Basic,        Basic,                           FColor::White,          "Basic") \
+	macro(Blueprint,    Blueprint,                       FColor(63, 127, 255),   "Blueprint") \
+	macro(Blueprint,    BlueprintGeneratedClass,         FColor(133, 173, 255),  "Blueprint Generated Class") \
+	macro(Blueprint,    CompiledBlueprintClass,          FColor(133, 173, 255),  "Compiled Blueprint Class") \
+	macro(Cinematics,   Cinematics,                      FColor::White,          "Cinematics") \
+	macro(Foliage,      Foliage,                         FColor::White,          "Foliage") \
+	macro(FX,           FX,                              FColor::White,          "FX") \
+	macro(FX,           Niagara,                         FColor::White,          "Niagara System") \
+	macro(Gameplay,     Gameplay,                        FColor::White,          "Gameplay") \
+	macro(Input,        Input,                           FColor::White,          "Input") \
+	macro(Material,     Material,                        FColor::White,          "Material") \
+	macro(Material,     MaterialFunction,                FColor(0, 172, 172),    "Material Function") \
+	macro(Material,     MaterialInstance,                FColor(0, 127, 0),      "Material Instance") \
+	macro(Media,        Media,                           FColor::White,          "Media") \
+	macro(Misc,         Misc,                            FColor::Red,            "Miscellaneous") \
+	macro(Misc,         DeletedObjectPlaceholder,        FColor::Red,            "Deleted Object Placeholder") \
+	macro(Misc,         DataAsset,                       FColor(195, 28, 82),    "Data Asset") \
+	macro(Misc,         EntityPrefab,                    FColor(102, 48, 172),   "Entity Prefab") \
+	macro(SkeletalMesh, SkeletalMesh,                    FColor::White,          "Skeletal Mesh") \
+	macro(StaticMesh,   StaticMesh,                      FColor::White,          "Static Mesh") \
+	macro(Physics,      Physics,                         FColor(200, 162, 108),  "Physics") \
+	macro(Physics,      PhysicalMaterialMask,            FColor(200, 162, 108),  "Physical Material Mask") \
+	macro(Texture,      Texture,                         FColor::White,          "Texture") \
+	macro(Texture,      Texture2D,                       FColor::White,          "Texture 2D") \
+	macro(Texture,      TextureCube,                     FColor::White,          "Texture Cube") \
+	macro(UI,           UI,                              FColor::Cyan,           "User Interface") \
+	macro(Verse,        Verse,                           FColor::Yellow,         "Verse") \
+	macro(Verse,        VerseDebugData,                  FColor::White,          "VerseDebugData") \
+	macro(Verse,        VerseDigest,                     FColor::White,          "VerseDigest") \
+	macro(World,        World,                           FColor::White,          "World") \
+	macro(World,        Level,                           FColor(246, 150, 0),    "Level") \
+
+namespace UE::Insights::ObjectProfiler
+{
+
+////////////////////////////////////////////////////////////////////////////////////////////////////
+// FAssetInfoProvider
+////////////////////////////////////////////////////////////////////////////////////////////////////
+
+FAssetInfoProvider::FAssetInfoProvider()
+{
+	InitCategories();
+}
+
+////////////////////////////////////////////////////////////////////////////////////////////////////
+
+FAssetInfoProvider::~FAssetInfoProvider()
+{
+}
+
+////////////////////////////////////////////////////////////////////////////////////////////////////
+
+void FAssetInfoProvider::InitCategories()
+{
+#define UE_INSIGHTS_ENUM(_Category, _ParentCategory, _Color, _DisplayName) \
+	{ \
+		FAssetInfoCategory& CategoryInfo = Categories.AddDefaulted_GetRef(); \
+		CategoryInfo.Name = FName(#_Category); \
+		CategoryInfo.ParentCategoryName = FName(#_ParentCategory); \
+		if (CategoryInfo.ParentCategoryName == CategoryInfo.Name) \
+		{ \
+			CategoryInfo.ParentCategoryName = NAME_None; \
+		} \
+		CategoryInfo.DisplayName = LOCTEXT("AssetCategory_" #_Category, _DisplayName); \
+		CategoryInfo.Color = FLinearColor(_Color); \
+		CategoryInfo.Icon = nullptr; /*_Icon;*/ \
+		Classes.Add(CategoryInfo.Name, Categories.Num() - 1); \
+	}
+	UE_INSIGHTS_ENUM_ASSET_CATEGORIES(UE_INSIGHTS_ENUM)
+#undef UE_INSIGHTS_ENUM
+
+	// Resolve pointers to parent categories.
+	for (FAssetInfoCategory& Category : Categories)
+	{
+		if (!Category.ParentCategoryName.IsNone())
+		{
+			const uint32* FoundIndex = Classes.Find(Category.ParentCategoryName);
+			check(FoundIndex != nullptr);
+			Category.ParentCategory = &Categories[*FoundIndex];
+		}
+	}
+}
+
+////////////////////////////////////////////////////////////////////////////////////////////////////
+
+const IAssetInfoCategory& FAssetInfoProvider::GetClassCategory(const FName InClassName) const
+{
+	const uint32* FoundIndex = Classes.Find(InClassName);
+	if (FoundIndex)
+	{
+		return Categories[*FoundIndex];
+	}
+	return DefaultCategory;
+}
+
+////////////////////////////////////////////////////////////////////////////////////////////////////
+
+const IAssetInfoCategory& FAssetInfoProvider::GetObjectCategory(const FName InClassName, const TCHAR* InObjectName, const TCHAR* InObjectPath) const
+{
+	return GetClassCategory(InClassName);
+}
+
+////////////////////////////////////////////////////////////////////////////////////////////////////
+
+FText FAssetInfoProvider::GetDisplayName(const FAssetData& InAssetData, const FName InClassName) const
+{
+	return GetClassCategory(InClassName).GetDisplayName();
+}
+
+////////////////////////////////////////////////////////////////////////////////////////////////////
+
+FLinearColor FAssetInfoProvider::GetColor(const FAssetData& InAssetData, const FName InClassName) const
+{
+	return GetClassCategory(InClassName).GetColor();
+}
+
+////////////////////////////////////////////////////////////////////////////////////////////////////
+
+const FSlateBrush* FAssetInfoProvider::GetIcon(const FAssetData& InAssetData, const FName InClassName) const
+{
+	return GetClassCategory(InClassName).GetIcon();
+}
+
+////////////////////////////////////////////////////////////////////////////////////////////////////
+
+const FSlateBrush* FAssetInfoProvider::GetThumbnail(const FAssetData& InAssetData, const FName InClassName) const
+{
+	return GetClassCategory(InClassName).GetIcon();
+}
+
+////////////////////////////////////////////////////////////////////////////////////////////////////
+
+} // namespace UE::Insights::ObjectProfiler
+
+#undef LOCTEXT_NAMESPACE
+#undef UE_INSIGHTS_ENUM_ASSET_CATEGORIES
+#undef UE_INSIGHTS_ENUM_ASSET_TYPES
+
+#endif // UE_INSIGHTS_WITH_CUSTOM_ASSET_INFO
diff --git a/Engine/Source/Developer/TraceInsights/Private/Insights/ObjectProfiler/AssetInfoProvider.h b/Engine/Source/Developer/TraceInsights/Private/Insights/ObjectProfiler/AssetInfoProvider.h
new file mode 100644
--- /dev/null
+++ b/Engine/Source/Developer/TraceInsights/Private/Insights/ObjectProfiler/AssetInfoProvider.h
@@ -0,0 +1,71 @@
+// Copyright Epic Games, Inc. All Rights Reserved.
+
+#pragma once
+
+#include "CoreTypes.h"
+
+#include "Insights/ObjectProfiler/IAssetInfoProvider.h"
+
+#ifndef UE_INSIGHTS_WITH_CUSTOM_ASSET_INFO
+#if defined(__INTELLISENSE__)
+#define UE_INSIGHTS_WITH_CUSTOM_ASSET_INFO 1
+#else
+#define UE_INSIGHTS_WITH_CUSTOM_ASSET_INFO !WITH_EDITOR
+#endif
+#endif
+
+#if UE_INSIGHTS_WITH_CUSTOM_ASSET_INFO
+
+namespace UE::Insights::ObjectProfiler
+{
+
+class FAssetInfoProvider : public IAssetInfoProvider
+{
+private:
+	class FAssetInfoCategory : public IAssetInfoCategory
+	{
+		friend class FAssetInfoProvider;
+
+	public:
+		FAssetInfoCategory() {}
+		virtual ~FAssetInfoCategory() = default;
+
+		virtual uint64 GetId() const override { return uint64(Name.GetComparisonIndex().ToUnstableInt()); }
+
+		virtual FText GetDisplayName() const override { return DisplayName; }
+		virtual FLinearColor GetColor() const override { return Color; }
+		virtual const FSlateBrush* GetIcon() const override { return Icon; }
+
+	private:
+		FName Name;
+		FName ParentCategoryName;
+		FAssetInfoCategory* ParentCategory = nullptr;
+		FText DisplayName;
+		FLinearColor Color;
+		const FSlateBrush* Icon = nullptr;
+	};
+
+public:
+	FAssetInfoProvider();
+	virtual ~FAssetInfoProvider();
+
+	virtual const IAssetInfoCategory& GetClassCategory(const FName InClassName) const override;
+	virtual const IAssetInfoCategory& GetObjectCategory(const FName InClassName, const TCHAR* InObjectName, const TCHAR* InObjectPath = nullptr) const override;
+
+	virtual FText GetDisplayName(const FAssetData& InAssetData, const FName InClassName) const override;
+	virtual FLinearColor GetColor(const FAssetData& InAssetData, const FName InClassName) const override;
+	virtual const FSlateBrush* GetIcon(const FAssetData& InAssetData, const FName InClassName) const override;
+	virtual const FSlateBrush* GetThumbnail(const FAssetData& InAssetData, const FName InClassName) const override;
+
+private:
+	void InitCategories();
+
+private:
+	TArray<FAssetInfoCategory> Categories;
+	TMap<FName, uint32> Classes; // ClassName -> Category Index
+	FAssetInfoCategory DefaultCategory;
+};
+
+} // namespace UE::Insights::ObjectProfiler
+
+#endif // UE_INSIGHTS_WITH_CUSTOM_ASSET_INFO
diff --git a/Engine/Source/Developer/TraceInsights/Private/Insights/ObjectProfiler/EditorAssetInfoProvider.cpp b/Engine/Source/Developer/TraceInsights/Private/Insights/ObjectProfiler/EditorAssetInfoProvider.cpp
new file mode 100644
--- /dev/null
+++ b/Engine/Source/Developer/TraceInsights/Private/Insights/ObjectProfiler/EditorAssetInfoProvider.cpp
@@ -0,0 +1,326 @@
+// Copyright Epic Games, Inc. All Rights Reserved.
+
+#include "EditorAssetInfoProvider.h"
+
+#if UE_INSIGHTS_WITH_EDITOR_ASSET_INFO
+
+#include "Styling/SlateIconFinder.h"
+
+// CoreUObject
+#include "AssetRegistry/AssetData.h"
+#include "UObject/UObjectGlobals.h"
+
+// AssetDefinition
+#include "Misc/AssetCategoryPath.h"
+#include "AssetDefinitionRegistry.h"
+
+// AssetRegistry
+#include "AssetRegistry/AssetRegistryModule.h"
+#include "AssetRegistry/IAssetRegistry.h"
+
+// AssetTools
+//#include "AssetTools.h"
+//#include "AssetToolsModule.h"
+//#include "IAssetTypeActions.h"
+
+// ContentBrowser
+#include "ContentBrowserModule.h"
+#include "IContentBrowserSingleton.h"
+
+// Editor
+#include "UnrealEdGlobals.h"
+#include "Editor/EditorEngine.h"
+#include "Editor/UnrealEdEngine.h"
+extern UNREALED_API class UEditorEngine* GEditor;
+
+// TraceInsights
+#include "Insights/ObjectProfiler/ViewModels/AssetInfoNode.h"
+#include "Insights/ObjectProfiler/ViewModels/ObjectNode.h"
+
+namespace UE::Insights::ObjectProfiler
+{
+
+////////////////////////////////////////////////////////////////////////////////////////////////////
+// FEditorAssetInfoCategory
+////////////////////////////////////////////////////////////////////////////////////////////////////
+
+FEditorAssetInfoProvider::FEditorAssetInfoCategory::FEditorAssetInfoCategory()
+	//: Path(FText::GetEmpty())
+{
+}
+
+////////////////////////////////////////////////////////////////////////////////////////////////////
+
+FText FEditorAssetInfoProvider::FEditorAssetInfoCategory::GetDisplayName() const
+{
+	if (Def)
+	{
+		return Def->GetAssetDisplayName();
+		//return Path.GetCategoryText();
+	}
+	return FText::GetEmpty();
+}
+
+////////////////////////////////////////////////////////////////////////////////////////////////////
+
+FLinearColor FEditorAssetInfoProvider::FEditorAssetInfoCategory::GetColor() const
+{
+	if (Def)
+	{
+		return Def->GetAssetColor();
+	}
+	return FLinearColor(1.0f, 1.0f, 1.0f, 1.0f);
+}
+
+////////////////////////////////////////////////////////////////////////////////////////////////////
+
+const FSlateBrush* FEditorAssetInfoProvider::FEditorAssetInfoCategory::GetIcon() const
+{
+	if (Def)
+	{
+		//if (const UClass* Class = FindObject<UClass>(nullptr, *ClassName.ToString()))
+		if (const UClass* Class = FindFirstObject<UClass>(ClassName.ToString()))
+		{
+			return FSlateIconFinder::FindIconBrushForClass(Class);
+
+			//FAssetData AssetData(Class);
+			//return Def->GetIconBrush(AssetData, Class->GetFName());
+		}
+	}
+	return nullptr;
+}
+
+////////////////////////////////////////////////////////////////////////////////////////////////////
+// FEditorAssetInfoProvider
+////////////////////////////////////////////////////////////////////////////////////////////////////
+
+FEditorAssetInfoProvider::FEditorAssetInfoProvider()
+	: AssetRegistryModule(FModuleManager::Get().LoadModuleChecked<FAssetRegistryModule>("AssetRegistry"))
+{
+}
+
+////////////////////////////////////////////////////////////////////////////////////////////////////
+
+FEditorAssetInfoProvider::~FEditorAssetInfoProvider()
+{
+}
+
+////////////////////////////////////////////////////////////////////////////////////////////////////
+
+const IAssetInfoCategory& FEditorAssetInfoProvider::GetClassCategory(FName InClassName) const
+{
+	const FEditorAssetInfoCategory* FoundCategory = Categories.Find(InClassName);
+	if (FoundCategory)
+	{
+		return *FoundCategory;
+	}
+
+	//if (const UClass* Class = FindObject<UClass>(nullptr, InClassName.ToString()))
+	if (const UClass* Class = FindFirstObject<UClass>(InClassName.ToString()))
+	{
+		if (const UAssetDefinition* Def = UAssetDefinitionRegistry::Get()->GetAssetDefinitionForClass(Class))
+		{
+			//TConstArrayView<FAssetCategoryPath> CategoryPaths = Def->GetAssetCategories();
+			//if (CategoryPaths.Num() > 0)
+			//{
+				FEditorAssetInfoCategory Category;
+				Category.ClassName = InClassName;
+				Category.Def = Def;
+
+				// TBD: Use category path (with sub-categories)...
+				//Category.Path = CategoryPaths[0]; // TBD: ???
+				//Category.DisplayName = CategoryPaths[0].GetCategoryText();
+				
+				// TBD: Maybe we should cache these:
+				//Category.DisplayName = Def->GetAssetDisplayName()
+				//Category.Color = Def->GetAssetColor();
+				//Category.Icon = FSlateIconFinder::FindIconBrushForClass(Class);
+
+				return Categories.Add(InClassName, Category);
+			//}
+		}
+	}
+	return DefaultCategory;
+}
+
+////////////////////////////////////////////////////////////////////////////////////////////////////
+
+const IAssetInfoCategory& FEditorAssetInfoProvider::GetObjectCategory(const FName InClassName, const TCHAR* InObjectName, const TCHAR* InObjectPath) const
+{
+	return GetClassCategory(InClassName);
+}
+
+////////////////////////////////////////////////////////////////////////////////////////////////////
+
+FText FEditorAssetInfoProvider::GetDisplayName(const FAssetData& InAssetData, const FName InClassName) const
+{
+	return GetClassCategory(InClassName).GetDisplayName();
+}
+
+////////////////////////////////////////////////////////////////////////////////////////////////////
+
+FLinearColor FEditorAssetInfoProvider::GetColor(const FAssetData& InAssetData, const FName InClassName) const
+{
+	return GetClassCategory(InClassName).GetColor();
+}
+
+////////////////////////////////////////////////////////////////////////////////////////////////////
+
+const FSlateBrush* FEditorAssetInfoProvider::GetIcon(const FAssetData& InAssetData, const FName InClassName) const
+{
+	const FEditorAssetInfoCategory& Category = static_cast<const FEditorAssetInfoCategory&>(GetClassCategory(InClassName));
+	if (Category.Def)
+	{
+		return Category.Def->GetIconBrush(InAssetData, InClassName);
+	}
+	return nullptr;
+}
+
+////////////////////////////////////////////////////////////////////////////////////////////////////
+
+const FSlateBrush* FEditorAssetInfoProvider::GetThumbnail(const FAssetData& InAssetData, const FName InClassName) const
+{
+	const FEditorAssetInfoCategory& Category = static_cast<const FEditorAssetInfoCategory&>(GetClassCategory(InClassName));
+	if (Category.Def)
+	{
+		return Category.Def->GetThumbnailBrush(InAssetData, InClassName);
+	}
+	return nullptr;
+}
+
+////////////////////////////////////////////////////////////////////////////////////////////////////
+
+bool FEditorAssetInfoProvider::GetAssetData(const FObjectNode& InObjectNode, FAssetData& OutAsset) const
+{
+	FString Path;
+
+	auto Package = InObjectNode.GetPackage();
+	if (Package)
+	{
+		Path = Package->GetObjectName(); //TBD: GetObjectPath()
+	}
+	else
+	{
+		Path = InObjectNode.GetObjectName(); //TBD: GetObjectPath()
+	}
+
+	auto Outer = InObjectNode.GetOuter();
+	bool Result = false;
+	do
+	{
+		OutAsset = AssetRegistryModule.Get().GetAssetByObjectPath(Path, false, false);
+		if (OutAsset.IsValid())
+		{
+			return true;
+		}
+		if (Outer.IsValid())
+		{
+			auto PackageOuter = Outer->GetPackage();
+			if (PackageOuter)
+			{
+				Path = PackageOuter->GetObjectPath();
+			}
+			else
+			{
+				Path = Outer->GetObjectName(); //TBD: GetObjectPath()
+			}
+			Outer = Outer->GetOuter();
+		}
+		else
+		{
+			break;
+		}
+	}
+	while (!OutAsset.IsValid());
+
+	return false;
+}
+
+////////////////////////////////////////////////////////////////////////////////////////////////////
+
+bool FEditorAssetInfoProvider::FindMatchedAsset(const TArray<FTableTreeNodePtr>& InSelectedNodes, FAssetData& OutAsset) const
+{
+	for (const FTableTreeNodePtr& SelectedNode : InSelectedNodes)
+	{
+		if (SelectedNode->Is<FObjectNode>())
+		{
+			const FObjectNode& ObjectNode = SelectedNode->As<FObjectNode>();
+			if (GetAssetData(ObjectNode, OutAsset))
+			{
+				return true;
+			}
+		}
+	}
+	return false;
+}
+
+////////////////////////////////////////////////////////////////////////////////////////////////////
+
+#if 0
+bool FEditorAssetInfoProvider::GetActors(const FAssetData& InAssetData, TArray<FSoftObjectPath>& OutActors) const
+{
+	TArray<FName> Names;
+	Names.Add(InAssetData.PackageName);
+
+	FBuildPackageDependencyMapBuilder Builder(GEditor->GetEditorWorldContext().World(), Names);
+
+	auto Map = Builder.Execute();
+	for (auto& Entry : Map)
+	{
+		for (auto& Actor : Entry.Value)
+		{
+			OutActors.Add(Actor.Path);
+		}
+	}
+
+	return !OutActors.IsEmpty();
+}
+#endif
+
+////////////////////////////////////////////////////////////////////////////////////////////////////
+
+bool FEditorAssetInfoProvider::BrowseToAsset(const FAssetInfoNode& InAssetInfo) const
+{
+	if (InAssetInfo.AssetData.IsValid())
+	{
+		if (FContentBrowserModule* ContentBrowserModule = FModuleManager::GetModulePtr<FContentBrowserModule>("ContentBrowser"))
+		{
+			ContentBrowserModule->Get().SyncBrowserToAssets({ InAssetInfo.AssetData });
+			return true;
+		}
+	}
+	return false;
+}
+
+////////////////////////////////////////////////////////////////////////////////////////////////////
+
+bool FEditorAssetInfoProvider::BrowseToActor(const FAssetInfoNode& InAssetInfo) const
+{
+	if (InAssetInfo.ActorPath.IsValid())
+	{
+		UObject* Obj = InAssetInfo.ActorPath.ResolveObject();
+		if (!Obj)
+		{
+			Obj = InAssetInfo.ActorPath.TryLoad();
+		}
+		if (AActor* AsActor = Cast<AActor>(Obj))
+		{
+			GEditor->SelectNone(false, true);
+			GEditor->SelectActor(AsActor, /*InSelected=*/true, /*bNotify=*/false, /*bSelectEvenIfHidden=*/true);
+			GEditor->NoteSelectionChange();
+			GEditor->MoveViewportCamerasToActor(*AsActor, false);
+			GUnrealEd->ShowActorProperties();
+			GUnrealEd->UpdateFloatingPropertyWindows();
+			return true;
+		}
+	}
+	return false;
+}
+
+////////////////////////////////////////////////////////////////////////////////////////////////////
+
+} // namespace UE::Insights::ObjectProfiler
+
+#undef USE_CONTENT_BROWSER
+
+#endif // UE_INSIGHTS_WITH_EDITOR_ASSET_INFO
diff --git a/Engine/Source/Developer/TraceInsights/Private/Insights/ObjectProfiler/EditorAssetInfoProvider.h b/Engine/Source/Developer/TraceInsights/Private/Insights/ObjectProfiler/EditorAssetInfoProvider.h
new file mode 100644
--- /dev/null
+++ b/Engine/Source/Developer/TraceInsights/Private/Insights/ObjectProfiler/EditorAssetInfoProvider.h
@@ -0,0 +1,82 @@
+// Copyright Epic Games, Inc. All Rights Reserved.
+
+#pragma once
+
+#include "CoreTypes.h"
+
+#include "Insights/ObjectProfiler/IAssetInfoProvider.h"
+
+#ifndef UE_INSIGHTS_WITH_EDITOR_ASSET_INFO
+#if defined(__INTELLISENSE__)
+#define UE_INSIGHTS_WITH_EDITOR_ASSET_INFO 1
+#else
+#define UE_INSIGHTS_WITH_EDITOR_ASSET_INFO WITH_EDITOR
+#endif
+#endif
+
+#if UE_INSIGHTS_WITH_EDITOR_ASSET_INFO
+
+#define UE_API TRACEINSIGHTS_API
+
+// AssetDefinition
+//#include "Misc/AssetCategoryPath.h"
+
+class UAssetDefinition;
+class FAssetRegistryModule;
+
+namespace UE::Insights::ObjectProfiler
+{
+
+class FEditorAssetInfoProvider : public IAssetInfoProvider
+{
+private:
+	class FEditorAssetInfoCategory : public IAssetInfoCategory
+	{
+		friend class FEditorAssetInfoProvider;
+
+	public:
+		FEditorAssetInfoCategory();
+		virtual ~FEditorAssetInfoCategory() = default;
+
+		virtual uint64 GetId() const override { return uint64(ClassName.GetComparisonIndex().ToUnstableInt()); }
+
+		virtual FText GetDisplayName() const override;
+		virtual FLinearColor GetColor() const override;
+		virtual const FSlateBrush* GetIcon() const override;
+
+	private:
+		FName ClassName;
+		TObjectPtr<const UAssetDefinition> Def;
+		//FAssetCategoryPath Path;
+	};
+
+public:
+	UE_API FEditorAssetInfoProvider();
+	UE_API virtual ~FEditorAssetInfoProvider();
+
+	UE_API virtual const IAssetInfoCategory& GetClassCategory(const FName InClassName) const override;
+	UE_API virtual const IAssetInfoCategory& GetObjectCategory(const FName InClassName, const TCHAR* InObjectName, const TCHAR* InObjectPath = nullptr) const override;
+
+	UE_API virtual FText GetDisplayName(const FAssetData& InAssetData, const FName InClassName) const override;
+	UE_API virtual FLinearColor GetColor(const FAssetData& InAssetData, const FName InClassName) const override;
+	UE_API virtual const FSlateBrush* GetIcon(const FAssetData& InAssetData, const FName InClassName) const override;
+	UE_API virtual const FSlateBrush* GetThumbnail(const FAssetData& InAssetData, const FName InClassName) const override;
+
+	UE_API virtual bool GetAssetData(const FObjectNode& InObjectNode, FAssetData& OutAsset) const override;
+	UE_API virtual bool FindMatchedAsset(const TArray<TSharedPtr<FTableTreeNode>>& InSelectedNodes, FAssetData& OutAsset) const override;
+
+	virtual bool GetActors(const FAssetData& InAssetData, TArray<FSoftObjectPath>& OutActors) const { return false; }
+	UE_API virtual bool BrowseToAsset(const FAssetInfoNode& InAssetInfo) const override;
+	UE_API virtual bool BrowseToActor(const FAssetInfoNode& InAssetInfo) const override;
+
+private:
+	FAssetRegistryModule& AssetRegistryModule;
+	mutable TMap<FName, FEditorAssetInfoCategory> Categories; // ClassName -> Category
+	FEditorAssetInfoCategory DefaultCategory;
+};
+
+} // namespace UE::Insights::ObjectProfiler
+
+#undef UE_API
+
+#endif // UE_INSIGHTS_WITH_EDITOR_ASSET_INFO
diff --git a/Engine/Source/Developer/TraceInsights/Public/Insights/ObjectProfiler/IAssetInfoProvider.h b/Engine/Source/Developer/TraceInsights/Public/Insights/ObjectProfiler/IAssetInfoProvider.h
new file mode 100644
--- /dev/null
+++ b/Engine/Source/Developer/TraceInsights/Public/Insights/ObjectProfiler/IAssetInfoProvider.h
@@ -0,0 +1,63 @@
+// Copyright Epic Games, Inc. All Rights Reserved.
+
+#pragma once
+
+#include "CoreTypes.h"
+
+#include "Containers/Array.h"
+#include "Containers/UnrealString.h"
+#include "Math/Color.h"
+#include "UObject/NameTypes.h"
+#include "UObject/SoftObjectPath.h"
+
+struct FAssetData;
+struct FSlateBrush;
+
+namespace UE::Insights
+{
+	class FTableTreeNode;
+}
+
+namespace UE::Insights::ObjectProfiler
+{
+
+class FObjectNode;
+struct FAssetInfoNode;
+
+class IAssetInfoCategory
+{
+public:
+	IAssetInfoCategory() = default;
+	virtual ~IAssetInfoCategory() = default;
+
+	virtual uint64 GetId() const = 0;
+
+	virtual FText GetDisplayName() const = 0;
+	virtual FLinearColor GetColor() const = 0;
+	virtual const FSlateBrush* GetIcon() const = 0;
+};
+
+class IAssetInfoProvider
+{
+public:
+	IAssetInfoProvider() = default;
+	virtual ~IAssetInfoProvider() = default;
+
+	virtual const IAssetInfoCategory& GetClassCategory(const FName InClassName) const = 0;
+	virtual const IAssetInfoCategory& GetObjectCategory(const FName InClassName, const TCHAR* InObjectName, const TCHAR* InObjectPath = nullptr) const = 0;
+
+	virtual FText GetDisplayName(const FAssetData& InAssetData, const FName InClassName) const = 0;
+	virtual FLinearColor GetColor(const FAssetData& InAssetData, const FName InClassName) const = 0;
+	virtual const FSlateBrush* GetIcon(const FAssetData& InAssetData, const FName InClassName) const = 0;
+	virtual const FSlateBrush* GetThumbnail(const FAssetData& InAssetData, const FName InClassName) const = 0;
+
+	virtual bool GetAssetData(const FObjectNode& InObjectNode, FAssetData& OutAsset) const { return false; }
+	virtual bool FindMatchedAsset(const TArray<TSharedPtr<FTableTreeNode>>& InSelectedNodes, FAssetData& OutAsset) const { return false; }
+
+	virtual bool GetActors(const FAssetData& InAssetData, TArray<FSoftObjectPath>& OutActors) const { return false; }
+
+	virtual bool BrowseToAsset(const FAssetInfoNode& InAssetInfo) const { return false; }
+	virtual bool BrowseToActor(const FAssetInfoNode& InAssetInfo) const { return false; }
+};
+
+} // namespace UE::Insights::ObjectProfiler
