[Insights]
* ObjectProfiler: Promoted Object Insights to a standalone major tab window. Replaced the previous approach of extending the Timing Profiler layout with a dedicated SObjectProfilerWindow (SMajorTabWindow) that owns its own dockable tabs for the Objects table tree view and Object Details panel.
* Object Details view: SObjectDetailsView now uses AssetInfoProvider.
* Object Details view: Added UI for selected object name, estimated size, thumbnail (via IAssetInfoProvider), placeholder validation issues section, and referenced-actors list.
* FObjectNode: Added InitObjectPath(): walks the outer chain to build a slash-separated full path. Must be called after all outer pointers are resolved.
* FObjectNode: Added MakeAssetData(): constructs FAssetData from package path and class info.

#rnx
#Synced-CL 52097700

diff --git a/Engine/Source/Developer/TraceInsights/Private/Insights/InsightsStyle.cpp b/Engine/Source/Developer/TraceInsights/Private/Insights/InsightsStyle.cpp
--- a/Engine/Source/Developer/TraceInsights/Private/Insights/InsightsStyle.cpp
+++ b/Engine/Source/Developer/TraceInsights/Private/Insights/InsightsStyle.cpp
@@ -242,6 +242,17 @@ void FInsightsStyle::FStyle::Initialize()
 	Set("Icons.BudgetSettings", new IMAGE_BRUSH_SVG("BudgetSettings", Icon16x16));
 	Set("Icons.TimeMarkerSettings", new IMAGE_BRUSH_SVG("TimeMarkerSettings", Icon16x16));
 
+	//////////////////////////////////////////////////
+	// Object Insights
+
+	Set("Icons.ObjectProfiler", new IMAGE_BRUSH_SVG("UObject", Icon16x16));
+
+	Set("Icons.ObjectTableTreeView", new EDITOR_IMAGE_BRUSH_SVG("Starship/Common/Spreadsheet", Icon16x16));
+	Set("Icons.ObjectTableTreeView.ToolBar", new EDITOR_IMAGE_BRUSH_SVG("Starship/Common/Spreadsheet", Icon20x20));
+
+	Set("Icons.ObjectDetailsView", new EDITOR_IMAGE_BRUSH_SVG("Starship/Common/Spreadsheet", Icon16x16));
+	Set("Icons.ObjectDetailsView.ToolBar", new EDITOR_IMAGE_BRUSH_SVG("Starship/Common/Spreadsheet", Icon20x20));
+
 	//////////////////////////////////////////////////
 	// Tasks
 
diff --git a/Engine/Source/Developer/TraceInsights/Private/Insights/ObjectProfiler/ObjectProfilerCommands.cpp b/Engine/Source/Developer/TraceInsights/Private/Insights/ObjectProfiler/ObjectProfilerCommands.cpp
new file mode 100644
--- /dev/null
+++ b/Engine/Source/Developer/TraceInsights/Private/Insights/ObjectProfiler/ObjectProfilerCommands.cpp
@@ -0,0 +1,105 @@
+// Copyright Epic Games, Inc. All Rights Reserved.
+
+#include "ObjectProfilerCommands.h"
+
+#include "DesktopPlatformModule.h"
+#include "Framework/Application/SlateApplication.h"
+
+// TraceInsights
+#include "Insights/InsightsManager.h"
+#include "Insights/InsightsStyle.h"
+#include "Insights/ObjectProfiler/ObjectProfilerManager.h"
+
+////////////////////////////////////////////////////////////////////////////////////////////////////
+
+#define LOCTEXT_NAMESPACE "UE::Insights::ObjectProfiler"
+
+namespace UE::Insights::ObjectProfiler
+{
+
+////////////////////////////////////////////////////////////////////////////////////////////////////
+// FObjectProfilerCommands
+////////////////////////////////////////////////////////////////////////////////////////////////////
+
+FObjectProfilerCommands::FObjectProfilerCommands()
+: TCommands<FObjectProfilerCommands>(
+	TEXT("ObjectProfilerCommands"),
+	NSLOCTEXT("Contexts", "ObjectProfilerCommands", "Insights - Object Insights"),
+	NAME_None,
+	FInsightsStyle::GetStyleSetName())
+{
+}
+
+////////////////////////////////////////////////////////////////////////////////////////////////////
+
+FObjectProfilerCommands::~FObjectProfilerCommands()
+{
+}
+
+////////////////////////////////////////////////////////////////////////////////////////////////////
+
+// UI_COMMAND takes long for the compiler to optimize
+UE_DISABLE_OPTIMIZATION_SHIP
+void FObjectProfilerCommands::RegisterCommands()
+{
+	UI_COMMAND(ToggleObjectTableTreeViewVisibility,
+		"Objects",
+		"Toggles the visibility of the main Objects table/tree view.",
+		EUserInterfaceActionType::ToggleButton,
+		FInputChord());
+
+	UI_COMMAND(ToggleObjectDetailsViewVisibility,
+		"Object Details",
+		"Toggles the visibility of the Object Details view.",
+		EUserInterfaceActionType::ToggleButton,
+		FInputChord());
+}
+UE_ENABLE_OPTIMIZATION_SHIP
+
+////////////////////////////////////////////////////////////////////////////////////////////////////
+// Toggle Commands
+////////////////////////////////////////////////////////////////////////////////////////////////////
+
+#define IMPLEMENT_TOGGLE_COMMAND(CmdName, IsEnabled, SetIsEnabled) \
+	\
+	void FObjectProfilerActionManager::Map_##CmdName##_Global()\
+	{\
+		This->CommandList->MapAction(This->GetCommands().CmdName, CmdName##_Custom());\
+	}\
+	\
+	const FUIAction FObjectProfilerActionManager::CmdName##_Custom() \
+	{\
+		FUIAction UIAction;\
+		UIAction.ExecuteAction = FExecuteAction::CreateRaw(this, &FObjectProfilerActionManager::CmdName##_Execute);\
+		UIAction.CanExecuteAction = FCanExecuteAction::CreateRaw(this, &FObjectProfilerActionManager::CmdName##_CanExecute);\
+		UIAction.GetActionCheckState = FGetActionCheckState::CreateRaw(this, &FObjectProfilerActionManager::CmdName##_GetCheckState);\
+		return UIAction;\
+	}\
+	\
+	void FObjectProfilerActionManager::CmdName##_Execute()\
+	{\
+		const bool b##IsEnabled = !This->IsEnabled();\
+		This->SetIsEnabled(b##IsEnabled);\
+	}\
+	\
+	bool FObjectProfilerActionManager::CmdName##_CanExecute() const\
+	{\
+		return FInsightsManager::Get()->GetSession().IsValid();\
+	}\
+	\
+	ECheckBoxState FObjectProfilerActionManager::CmdName##_GetCheckState() const\
+	{\
+		const bool b##IsEnabled = This->IsEnabled();\
+		return b##IsEnabled ? ECheckBoxState::Checked : ECheckBoxState::Unchecked;\
+	}
+
+IMPLEMENT_TOGGLE_COMMAND(ToggleObjectTableTreeViewVisibility, IsObjectTableTreeViewVisible, ShowHideObjectTableTreeView)
+IMPLEMENT_TOGGLE_COMMAND(ToggleObjectDetailsViewVisibility, IsObjectDetailsViewVisible, ShowHideObjectDetailsView)
+
+#undef IMPLEMENT_TOGGLE_COMMAND
+
+////////////////////////////////////////////////////////////////////////////////////////////////////
+
+} // namespace UE::Insights::ObjectProfiler
+
+#undef LOCTEXT_NAMESPACE
diff --git a/Engine/Source/Developer/TraceInsights/Private/Insights/ObjectProfiler/ObjectProfilerCommands.h b/Engine/Source/Developer/TraceInsights/Private/Insights/ObjectProfiler/ObjectProfilerCommands.h
new file mode 100644
--- /dev/null
+++ b/Engine/Source/Developer/TraceInsights/Private/Insights/ObjectProfiler/ObjectProfilerCommands.h
@@ -0,0 +1,79 @@
+// Copyright Epic Games, Inc. All Rights Reserved.
+
+#pragma once
+
+#include "CoreTypes.h"
+
+#include "Framework/Commands/UIAction.h"
+#include "Framework/Commands/Commands.h"
+
+namespace UE::Insights::ObjectProfiler
+{
+
+////////////////////////////////////////////////////////////////////////////////////////////////////
+/**
+ * Class that holds all profiler commands.
+ */
+class FObjectProfilerCommands : public TCommands<FObjectProfilerCommands>
+{
+public:
+	FObjectProfilerCommands();
+	virtual ~FObjectProfilerCommands();
+	virtual void RegisterCommands() override;
+
+public:
+	//////////////////////////////////////////////////
+	// Global commands need to implement following method:
+	//     void Map_<CommandName>_Global();
+	// Custom commands needs to implement also the following method:
+	//     const FUIAction <CommandName>_Custom(...) const;
+	//////////////////////////////////////////////////
+
+	/** Toggles visibility for the Objects table/tree view. Global and custom command. */
+	TSharedPtr<FUICommandInfo> ToggleObjectTableTreeViewVisibility;
+
+	/** Toggles visibility for the Object Details view. Global and custom command. */
+	TSharedPtr<FUICommandInfo> ToggleObjectDetailsViewVisibility;
+};
+
+////////////////////////////////////////////////////////////////////////////////////////////////////
+/**
+ * Class that provides helper functions for the commands to avoid cluttering profiler manager with many small functions.
+ * Can't contain any variables. Directly operates on the profiler manager instance.
+ */
+class FObjectProfilerActionManager
+{
+	friend class FObjectProfilerManager;
+
+private:
+	/** Private constructor. */
+	FObjectProfilerActionManager(class FObjectProfilerManager* Instance)
+		: This(Instance)
+	{}
+
+	//////////////////////////////////////////////////
+	// Toggle Commands
+
+#define DECLARE_TOGGLE_COMMAND(CmdName)\
+public:\
+	void Map_##CmdName##_Global(); /**< Maps UI command info CmdName with the specified UI command list. */\
+	const FUIAction CmdName##_Custom(); /**< UI action for CmdName command. */\
+protected:\
+	void CmdName##_Execute(); /**< Handles FExecuteAction for CmdName. */\
+	bool CmdName##_CanExecute() const; /**< Handles FCanExecuteAction for CmdName. */\
+	ECheckBoxState CmdName##_GetCheckState() const; /**< Handles FGetActionCheckState for CmdName. */
+
+	DECLARE_TOGGLE_COMMAND(ToggleObjectTableTreeViewVisibility)
+	DECLARE_TOGGLE_COMMAND(ToggleObjectDetailsViewVisibility)
+#undef DECLARE_TOGGLE_COMMAND
+
+	//////////////////////////////////////////////////
+
+protected:
+	/** Reference to the global instance of the profiler manager. */
+	class FObjectProfilerManager* This;
+};
+
+////////////////////////////////////////////////////////////////////////////////////////////////////
+
+} // namespace UE::Insights::ObjectProfiler
diff --git a/Engine/Source/Developer/TraceInsights/Private/Insights/ObjectProfiler/ObjectProfilerManager.cpp b/Engine/Source/Developer/TraceInsights/Private/Insights/ObjectProfiler/ObjectProfilerManager.cpp
--- a/Engine/Source/Developer/TraceInsights/Private/Insights/ObjectProfiler/ObjectProfilerManager.cpp
+++ b/Engine/Source/Developer/TraceInsights/Private/Insights/ObjectProfiler/ObjectProfilerManager.cpp
@@ -2,33 +2,40 @@
 
 #include "ObjectProfilerManager.h"
 
-#include "Features/IModularFeatures.h"
-#include "Framework/Docking/TabManager.h"
+#include "HAL/IConsoleManager.h"
+#include "MessageLogModule.h"
 #include "Modules/ModuleManager.h"
 #include "Widgets/Docking/SDockTab.h"
+#include "WorkspaceMenuStructure.h"
+#include "WorkspaceMenuStructureModule.h"
 
 // TraceServices
 #include "Common/ProviderLock.h"
 #include "TraceServices/Model/ObjectProvider.h"
 
 // TraceInsights
+#include "Insights/Common/InsightsMenuBuilder.h"
+#include "Insights/InsightsManager.h"
 #include "Insights/InsightsStyle.h"
 #include "Insights/ObjectProfiler/ViewModels/ObjectTable.h"
-#include "Insights/ObjectProfiler/Widgets/SMemoryView.h"
+#include "Insights/ObjectProfiler/Widgets/SObjectProfilerWindow.h"
 
 #define LOCTEXT_NAMESPACE "UE::Insights::ObjectProfiler"
 
+// For Object Insights in editor preview.
+static bool GEnableObjectInsightsInEditor = true;
+static FAutoConsoleVariableRef CVAREnableObjectInsights(
+	TEXT("Trace.EnableObjectInsights"),
+	GEnableObjectInsightsInEditor,
+	TEXT("Enables the Object Insights feature in the Editor."),
+	ECVF_Default
+);
+
 namespace UE::Insights::ObjectProfiler
 {
 
 DEFINE_LOG_CATEGORY(LogObjectProfiler);
 
-////////////////////////////////////////////////////////////////////////////////////////////////////
-
-const FName FObjectProfilerTabs::ObjectsTableTreeViewTabId(TEXT("ObjectsTableTreeView"));
-
-////////////////////////////////////////////////////////////////////////////////////////////////////
-
 TSharedPtr<FObjectProfilerManager> FObjectProfilerManager::Instance = nullptr;
 
 ////////////////////////////////////////////////////////////////////////////////////////////////////
@@ -56,8 +63,10 @@ TSharedPtr<FObjectProfilerManager> FObjectProfilerManager::CreateInstance()
 ////////////////////////////////////////////////////////////////////////////////////////////////////
 
 FObjectProfilerManager::FObjectProfilerManager(TSharedRef<FUICommandList> InCommandList)
-	: bIsInitialized(false)
-	, bIsAvailable(false)
+	: CommandList(InCommandList)
+	, ActionManager(this)
+	, MajorTabId(SObjectProfilerWindow::ObjectProfilerTabId)
+	, LogListingName("ObjectInsights")
 {
 }
 
@@ -72,15 +81,16 @@ void FObjectProfilerManager::Initialize(IUnrealInsightsModule& InsightsModule)
 	}
 	bIsInitialized = true;
 
+	UE_LOGF(LogObjectProfiler, Log, "Initialize");
+
 	// Register tick functions.
 	OnTick = FTickerDelegate::CreateSP(this, &FObjectProfilerManager::Tick);
 	OnTickHandle = FTSTicker::GetCoreTicker().AddTicker(OnTick, 0.0f);
 
-	FOnRegisterMajorTabExtensions* TimingProfilerLayoutExtension = InsightsModule.FindMajorTabLayoutExtension(FInsightsManagerTabs::TimingProfilerTabId);
-	if (TimingProfilerLayoutExtension)
-	{
-		TimingProfilerLayoutExtension->AddRaw(this, &FObjectProfilerManager::RegisterTimingProfilerLayoutExtensions);
-	}
+	FObjectProfilerCommands::Register();
+	BindCommands();
+
+	InsightsModule.OnRegisterMajorTabExtension(MajorTabId);
 
 	FInsightsManager::Get()->GetSessionChangedEvent().AddSP(this, &FObjectProfilerManager::OnSessionChanged);
 	OnSessionChanged();
@@ -96,12 +106,26 @@ void FObjectProfilerManager::Shutdown()
 	}
 	bIsInitialized = false;
 
+	// If the MessageLog module was already unloaded as part of the global Shutdown process, do not load it again.
+	if (FModuleManager::Get().IsModuleLoaded("MessageLog"))
+	{
+		FMessageLogModule& MessageLogModule = FModuleManager::LoadModuleChecked<FMessageLogModule>("MessageLog");
+		if (MessageLogModule.IsRegisteredLogListing(GetLogListingName()))
+		{
+			MessageLogModule.UnregisterLogListing(GetLogListingName());
+		}
+	}
+
 	FInsightsManager::Get()->GetSessionChangedEvent().RemoveAll(this);
 
+	FObjectProfilerCommands::Unregister();
+
 	// Unregister tick function.
 	FTSTicker::RemoveTicker(OnTickHandle);
 
 	FObjectProfilerManager::Instance.Reset();
+
+	UE_LOGF(LogObjectProfiler, Log, "Shutdown");
 }
 
 ////////////////////////////////////////////////////////////////////////////////////////////////////
@@ -113,52 +137,177 @@ FObjectProfilerManager::~FObjectProfilerManager()
 
 ////////////////////////////////////////////////////////////////////////////////////////////////////
 
+void FObjectProfilerManager::BindCommands()
+{
+	ActionManager.Map_ToggleObjectTableTreeViewVisibility_Global();
+	ActionManager.Map_ToggleObjectDetailsViewVisibility_Global();
+}
+
+////////////////////////////////////////////////////////////////////////////////////////////////////
+
 void FObjectProfilerManager::RegisterMajorTabs(IUnrealInsightsModule& InsightsModule)
 {
+	bool bShouldRegister = true;
+
+#if WITH_EDITOR
+	bShouldRegister = GEnableObjectInsightsInEditor;
+#endif
+
+	if (!bIsTabRegistered && bShouldRegister)
+	{
+		bIsTabRegistered = true;
+		const FInsightsMajorTabConfig& Config = InsightsModule.FindMajorTabConfig(MajorTabId);
+		if (Config.bIsAvailable)
+		{
+			// Register tab spawner for the Object Insights.
+			FTabSpawnerEntry& TabSpawnerEntry = FGlobalTabmanager::Get()->RegisterNomadTabSpawner(MajorTabId,
+					FOnSpawnTab::CreateRaw(this, &FObjectProfilerManager::SpawnTab),
+					FCanSpawnTab::CreateRaw(this, &FObjectProfilerManager::CanSpawnTab))
+				.SetDisplayName(Config.TabLabel.IsSet() ? Config.TabLabel.GetValue() : LOCTEXT("ObjectProfilerTabTitle", "Object Insights"))
+				.SetTooltipText(Config.TabTooltip.IsSet() ? Config.TabTooltip.GetValue() : LOCTEXT("ObjectProfilerTooltipText", "Open the Object Insights tab."))
+				.SetIcon(Config.TabIcon.IsSet() ? Config.TabIcon.GetValue() : FSlateIcon(FInsightsStyle::GetStyleSetName(), "Icons.ObjectProfiler"));
+
+			TSharedRef<FWorkspaceItem> Group = Config.WorkspaceGroup.IsValid() ? Config.WorkspaceGroup.ToSharedRef() : FInsightsManager::Get()->GetInsightsMenuBuilder()->GetInsightsToolsGroup();
+			TabSpawnerEntry.SetGroup(Group);
+		}
+	}
 }
 
 ////////////////////////////////////////////////////////////////////////////////////////////////////
 
 void FObjectProfilerManager::UnregisterMajorTabs()
 {
+	FGlobalTabmanager::Get()->UnregisterNomadTabSpawner(MajorTabId);
+	bIsTabRegistered = false;
+}
+
+////////////////////////////////////////////////////////////////////////////////////////////////////
+
+TSharedRef<SDockTab> FObjectProfilerManager::SpawnTab(const FSpawnTabArgs& Args)
+{
+	const TSharedRef<SDockTab> DockTab = SNew(SDockTab)
+		.TabRole(ETabRole::NomadTab);
+
+	// Register OnTabClosed to handle Object profiler manager shutdown.
+	DockTab->SetOnTabClosed(SDockTab::FOnTabClosedCallback::CreateRaw(this, &FObjectProfilerManager::OnTabClosed));
+
+	// Create the SObjectProfilerWindow widget.
+	TSharedRef<SObjectProfilerWindow> Window = SNew(SObjectProfilerWindow, DockTab, Args.GetOwnerWindow());
+	DockTab->SetContent(Window);
+
+	AssignProfilerWindow(Window);
+
+	return DockTab;
+}
+
+////////////////////////////////////////////////////////////////////////////////////////////////////
+
+bool FObjectProfilerManager::CanSpawnTab(const FSpawnTabArgs& Args) const
+{
+#if WITH_EDITOR
+	return true;
+#else
+	return bIsAvailable;
+#endif
+}
+
+////////////////////////////////////////////////////////////////////////////////////////////////////
+
+void FObjectProfilerManager::OnTabClosed(TSharedRef<SDockTab> TabBeingClosed)
+{
+	OnWindowClosedEvent();
+
+	RemoveProfilerWindow();
+
+	// Disable TabClosed delegate.
+	TabBeingClosed->SetOnTabClosed(SDockTab::FOnTabClosedCallback());
+}
+
+////////////////////////////////////////////////////////////////////////////////////////////////////
+
+const TSharedRef<FUICommandList> FObjectProfilerManager::GetCommandList() const
+{
+	return CommandList;
+}
+
+////////////////////////////////////////////////////////////////////////////////////////////////////
+
+const FObjectProfilerCommands& FObjectProfilerManager::GetCommands()
+{
+	return FObjectProfilerCommands::Get();
+}
+
+////////////////////////////////////////////////////////////////////////////////////////////////////
+
+FObjectProfilerActionManager& FObjectProfilerManager::GetActionManager()
+{
+	return FObjectProfilerManager::Instance->ActionManager;
 }
 
 ////////////////////////////////////////////////////////////////////////////////////////////////////
 
 bool FObjectProfilerManager::Tick(float DeltaTime)
 {
+#if WITH_EDITOR
+	if (!bIsTabRegistered && GEnableObjectInsightsInEditor)
+	{
+		IUnrealInsightsModule& TraceInsightsModule = FModuleManager::LoadModuleChecked<IUnrealInsightsModule>("TraceInsights");
+		RegisterMajorTabs(TraceInsightsModule);
+	}
+#endif
+
 	// Check if the session has UObject snapshots, but not too often.
 	if (!bIsAvailable && AvailabilityCheck.Tick())
 	{
 		bool bShouldBeAvailable = false;
 
-		TSharedPtr<const TraceServices::IAnalysisSession> Session = FInsightsManager::Get()->GetSession();
+		TSharedPtr<FInsightsManager> InsightsManager = FInsightsManager::Get();
+		check(InsightsManager.IsValid());
+
+		TSharedPtr<const TraceServices::IAnalysisSession> Session = InsightsManager->GetSession();
 		if (Session.IsValid())
 		{
-			const TraceServices::IObjectProvider* ObjectProvider = TraceServices::ReadObjectProvider(*Session.Get());
-			if (ObjectProvider)
 			{
-				TraceServices::FProviderReadScopeLock ProviderReadScope(*ObjectProvider);
-
-				TSharedPtr<FTabManager> TabManagerShared = TimingTabManager.Pin();
-				if (ObjectProvider->GetNumSnapshots() > 0 && TabManagerShared.IsValid())
-				{
-					bIsAvailable = true;
-					TabManagerShared->TryInvokeTab(FObjectProfilerTabs::ObjectsTableTreeViewTabId, true);
-				}
-
+				TraceServices::FAnalysisSessionReadScope SessionReadScope(*Session.Get());
 				if (Session->IsAnalysisComplete())
 				{
 					// Never check again during this session.
 					AvailabilityCheck.Disable();
 				}
 			}
+
+			const TraceServices::IObjectProvider* ObjectProvider = TraceServices::ReadObjectProvider(*Session.Get());
+			if (ObjectProvider)
+			{
+				TraceServices::FProviderReadScopeLock ProviderReadScope(*ObjectProvider);
+				if (ObjectProvider->GetNumSnapshots() > 0)
+				{
+					bShouldBeAvailable = true;
+				}
+			}
 		}
 		else
 		{
 			// Do not check again until the next session changed event (see OnSessionChanged).
 			AvailabilityCheck.Disable();
 		}
+
+		if (bShouldBeAvailable)
+		{
+			bIsAvailable = true;
+
+			FMessageLogModule& MessageLogModule = FModuleManager::LoadModuleChecked<FMessageLogModule>("MessageLog");
+			MessageLogModule.RegisterLogListing(GetLogListingName(), LOCTEXT("ObjectInsights", "Object Insights"));
+			MessageLogModule.EnableMessageLogDisplay(true);
+
+#if !WITH_EDITOR
+			if (FGlobalTabmanager::Get()->HasTabSpawner(MajorTabId))
+			{
+				UE_LOGF(LogObjectProfiler, Log, "Opening the \"Object Insights\" tab...");
+				FGlobalTabmanager::Get()->TryInvokeTab(MajorTabId);
+			}
+#endif
+		}
 	}
 
 	return true;
@@ -168,6 +317,8 @@ bool FObjectProfilerManager::Tick(float DeltaTime)
 
 void FObjectProfilerManager::OnSessionChanged()
 {
+	UE_LOGF(LogObjectProfiler, Log, "OnSessionChanged");
+
 	bIsAvailable = false;
 	if (FInsightsManager::Get()->GetSession().IsValid())
 	{
@@ -177,83 +328,64 @@ void FObjectProfilerManager::OnSessionChanged()
 	{
 		AvailabilityCheck.Disable();
 	}
-}
-
-////////////////////////////////////////////////////////////////////////////////////////////////////
 
-void FObjectProfilerManager::RegisterTimingProfilerLayoutExtensions(FInsightsMajorTabExtender& InOutExtender)
-{
-	TimingTabManager = InOutExtender.GetTabManager();
-
-	FMinorTabConfig& MinorTabConfig = InOutExtender.AddMinorTabConfig();
-	MinorTabConfig.TabId = FObjectProfilerTabs::ObjectsTableTreeViewTabId;
-	MinorTabConfig.TabLabel = LOCTEXT("ObjectsTableTreeViewTabTitle", "Objects");
-	MinorTabConfig.TabTooltip = LOCTEXT("ObjectsTableTreeViewTabTitleTooltip", "Opens the Objects Table Tree View tab.");
-	MinorTabConfig.TabIcon = FSlateIcon(FInsightsStyle::GetStyleSetName(), "Icons.ObjectsView");
-	MinorTabConfig.OnSpawnTab = FOnSpawnTab::CreateRaw(this, &FObjectProfilerManager::SpawnTab_ObjectsTableTreeView);
-	MinorTabConfig.CanSpawnTab = FCanSpawnTab::CreateRaw(this, &FObjectProfilerManager::CanSpawnTab_ObjectsTableTreeView);
-	MinorTabConfig.WorkspaceGroup = InOutExtender.GetWorkspaceGroup();
-
-	InOutExtender.GetLayoutExtender().ExtendLayout(FTimingProfilerTabs::StatsCountersID
-		, ELayoutExtensionPosition::After
-		, FTabManager::FTab(FObjectProfilerTabs::ObjectsTableTreeViewTabId, ETabState::ClosedTab));
+	TSharedPtr<SObjectProfilerWindow> Wnd = GetProfilerWindow();
+	if (Wnd)
+	{
+		Wnd->Reset();
+	}
 }
 
 ////////////////////////////////////////////////////////////////////////////////////////////////////
 
-TSharedRef<SDockTab> FObjectProfilerManager::SpawnTab_ObjectsTableTreeView(const FSpawnTabArgs& Args)
+void FObjectProfilerManager::ShowHideObjectTableTreeView(bool bIsVisible)
 {
-	const TSharedRef<SDockTab> DockTab = SNew(SDockTab)
-		.ShouldAutosize(false)
-		.TabRole(ETabRole::PanelTab)
-		[
-			SAssignNew(MemoryView, SMemoryView)
-		];
-
-	DockTab->SetOnTabClosed(SDockTab::FOnTabClosedCallback::CreateRaw(this, &FObjectProfilerManager::OnObjectsTableTreeViewTabClosed));
+	bIsObjectTableTreeViewVisible = bIsVisible;
 
-	return DockTab;
+	TSharedPtr<SObjectProfilerWindow> Wnd = GetProfilerWindow();
+	if (Wnd)
+	{
+		Wnd->ShowHideTab(FObjectProfilerTabs::ObjectTableTreeViewId, bIsVisible);
+	}
 }
 
 ////////////////////////////////////////////////////////////////////////////////////////////////////
 
-bool FObjectProfilerManager::CanSpawnTab_ObjectsTableTreeView(const FSpawnTabArgs& Args)
+void FObjectProfilerManager::ShowHideObjectDetailsView(bool bIsVisible)
 {
-	return bIsAvailable;
+	bIsObjectDetailsViewVisible = bIsVisible;
+
+	TSharedPtr<SObjectProfilerWindow> Wnd = GetProfilerWindow();
+	if (Wnd)
+	{
+		Wnd->ShowHideTab(FObjectProfilerTabs::ObjectDetailsViewId, bIsVisible);
+	}
 }
 
 ////////////////////////////////////////////////////////////////////////////////////////////////////
 
-void FObjectProfilerManager::OnObjectsTableTreeViewTabClosed(TSharedRef<SDockTab> TabBeingClosed)
+void FObjectProfilerManager::ShowHideSegmentedBarGraph(bool bIsVisible)
 {
-	if (MemoryView.IsValid())
+	bIsSegmentedBarGraphVisible = bIsVisible;
+
+	TSharedPtr<SObjectProfilerWindow> Wnd = GetProfilerWindow();
+	if (Wnd)
 	{
-		MemoryView->OnClose();
+		Wnd->ShowHideSegmentedBarGraph(bIsVisible);
 	}
-
-	MemoryView.Reset();
 }
 
 ////////////////////////////////////////////////////////////////////////////////////////////////////
 
 void FObjectProfilerManager::OnWindowClosedEvent()
 {
-	if (MemoryView.IsValid())
-	{
-		MemoryView->OnClose();
-	}
+}
 
-	TSharedPtr<FTabManager> TimingTabManagerSharedPtr = TimingTabManager.Pin();
+////////////////////////////////////////////////////////////////////////////////////////////////////
 
-	if (TimingTabManagerSharedPtr.IsValid())
-	{
-		TSharedPtr<SDockTab> Tab = TimingTabManagerSharedPtr->FindExistingLiveTab(FObjectProfilerTabs::ObjectsTableTreeViewTabId);
-		if (Tab.IsValid())
-		{
-			Tab->RequestCloseTab();
-			Tab->RemoveTabFromParent();
-		}
-	}
+bool FObjectProfilerManager::Exec(const TCHAR* Cmd, FOutputDevice& Ar)
+{
+	return false;
 }
 
 ////////////////////////////////////////////////////////////////////////////////////////////////////
diff --git a/Engine/Source/Developer/TraceInsights/Private/Insights/ObjectProfiler/ObjectProfilerManager.h b/Engine/Source/Developer/TraceInsights/Private/Insights/ObjectProfiler/ObjectProfilerManager.h
--- a/Engine/Source/Developer/TraceInsights/Private/Insights/ObjectProfiler/ObjectProfilerManager.h
+++ b/Engine/Source/Developer/TraceInsights/Private/Insights/ObjectProfiler/ObjectProfilerManager.h
@@ -6,28 +6,22 @@
 
 #include "Containers/Ticker.h"
 #include "Framework/Commands/UICommandList.h"
-#include "Templates/SharedPointer.h"
 #include "Logging/LogMacros.h"
+#include "Templates/SharedPointer.h"
 
 // TraceInsightsCore
 #include "InsightsCore/Common/AvailabilityCheck.h"
 
 // TraceInsights
-#include "Insights/InsightsManager.h"
 #include "Insights/IUnrealInsightsModule.h"
+#include "Insights/ObjectProfiler/ObjectProfilerCommands.h"
 
 namespace UE::Insights::ObjectProfiler
 {
 
-DECLARE_LOG_CATEGORY_EXTERN(LogObjectProfiler, Log, All);
+class SObjectProfilerWindow;
 
-class SMemoryView;
-
-struct FObjectProfilerTabs
-{
-	// Tab identifiers
-	static const FName ObjectsTableTreeViewTabId;
-};
+DECLARE_LOG_CATEGORY_EXTERN(LogObjectProfiler, Log, All);
 
 ////////////////////////////////////////////////////////////////////////////////////////////////////
 /**
@@ -35,6 +29,8 @@ struct FObjectProfilerTabs
  */
 class FObjectProfilerManager : public TSharedFromThis<FObjectProfilerManager>, public IInsightsComponent
 {
+	friend class FObjectProfilerActionManager;
+
 public:
 	/** Creates the Object Profiler manager, only one instance can exist. */
 	FObjectProfilerManager(TSharedRef<FUICommandList> InCommandList);
@@ -42,7 +38,7 @@ public:
 	/** Virtual destructor. */
 	virtual ~FObjectProfilerManager();
 
-	/** Creates an instance of the ObjectProfiler manager. */
+	/** Creates an instance of the Object Profiler manager. */
 	static TSharedPtr<FObjectProfilerManager> CreateInstance();
 
 	/**
@@ -54,6 +50,8 @@ public:
 	 */
 	static TSharedPtr<FObjectProfilerManager> Get();
 
+	bool IsAvailable() const { return bIsAvailable; }
+
 	//////////////////////////////////////////////////
 	// IInsightsComponent
 
@@ -62,27 +60,85 @@ public:
 	virtual void RegisterMajorTabs(IUnrealInsightsModule& InsightsModule) override;
 	virtual void UnregisterMajorTabs() override;
 	virtual void OnWindowClosedEvent() override;
+	virtual bool Exec(const TCHAR* Cmd, FOutputDevice& Ar) override;
 
 	//////////////////////////////////////////////////
 
-	TSharedRef<SDockTab> SpawnTab_ObjectsTableTreeView(const FSpawnTabArgs& Args);
-	bool CanSpawnTab_ObjectsTableTreeView(const FSpawnTabArgs& Args);
-	void OnObjectsTableTreeViewTabClosed(TSharedRef<SDockTab> TabBeingClosed);
+	/** @return UI command list for the Object Profiler manager. */
+	const TSharedRef<FUICommandList> GetCommandList() const;
 
-	//////////////////////////////////////////////////
+	/** @return an instance of the Object Profiler commands. */
+	static const FObjectProfilerCommands& GetCommands();
+
+	/** @return an instance of the Object Profiler action manager. */
+	static FObjectProfilerActionManager& GetActionManager();
+
+	/**
+	 * Converts profiler window weak pointer to a shared pointer and returns it.
+	 * Make sure the returned pointer is valid before trying to dereference it.
+	 */
+	TSharedPtr<SObjectProfilerWindow> GetProfilerWindow() const
+	{
+		return ProfilerWindowWeakPtr.Pin();
+	}
+
+	const FName& GetMajorTabId() const { return MajorTabId; }
+	const FName& GetLogListingName() const { return LogListingName; }
+
+	////////////////////////////////////////////////////////////////////////////////////////////////////
+	// Getters and setters used by Toggle Commands.
+
+	/** @return true, if the Objects table/tree view is visible */
+	bool IsObjectTableTreeViewVisible() const { return bIsObjectTableTreeViewVisible; }
+	void SetObjectTableTreeViewVisible(bool bIsVisible) { bIsObjectTableTreeViewVisible = bIsVisible; }
+	void ShowHideObjectTableTreeView(bool bIsVisible);
+
+	/** @return true, if the Object Details view is visible */
+	bool IsObjectDetailsViewVisible() const { return bIsObjectDetailsViewVisible; }
+	void SetObjectDetailsViewVisible(bool bIsVisible) { bIsObjectDetailsViewVisible = bIsVisible; }
+	void ShowHideObjectDetailsView(bool bIsVisible);
 
-	bool IsAvailable() { return bIsAvailable; }
+	/** @return true, if the Segmented Bar Graph is visible (in toolbar) */
+	bool IsSegmentedBarGraphVisible() const { return bIsSegmentedBarGraphVisible; }
+	void SetSegmentedBarGraphVisible(bool bIsVisible) { bIsSegmentedBarGraphVisible = bIsVisible; }
+	void ShowHideSegmentedBarGraph(bool bIsVisible);
+
+	////////////////////////////////////////////////////////////////////////////////////////////////////
 
 	void OnSessionChanged();
 
 private:
+	/** Binds our UI commands to delegates. */
+	void BindCommands();
+
+	/** Called to spawn the Object Profiler major tab. */
+	TSharedRef<SDockTab> SpawnTab(const FSpawnTabArgs& Args);
+
+	bool CanSpawnTab(const FSpawnTabArgs& Args) const;
+
+	/** Callback called when the Object Profiler major tab is closed. */
+	void OnTabClosed(TSharedRef<SDockTab> TabBeingClosed);
+
+	/** Updates this manager, done through FCoreTicker. */
 	bool Tick(float DeltaTime);
 
-	void RegisterTimingProfilerLayoutExtensions(FInsightsMajorTabExtender& InOutExtender);
+	void AssignProfilerWindow(const TSharedRef<SObjectProfilerWindow>& InProfilerWindow)
+	{
+		ProfilerWindowWeakPtr = InProfilerWindow;
+	}
+
+	void RemoveProfilerWindow()
+	{
+		ProfilerWindowWeakPtr.Reset();
+	}
 
 private:
-	bool bIsInitialized;
-	bool bIsAvailable;
+	bool bIsInitialized = false;
+	bool bIsAvailable = false;
+
+	/** If the tab is registered in the tab manager. */
+	bool bIsTabRegistered = false;
+
 	FAvailabilityCheck AvailabilityCheck;
 
 	/** The delegate to be invoked when this manager ticks. */
@@ -91,12 +147,32 @@ private:
 	/** Handle to the registered OnTick. */
 	FTSTicker::FDelegateHandle OnTickHandle;
 
-	/** A shared pointer to the global instance of the Object Profiler manager. */
-	static TSharedPtr<FObjectProfilerManager> Instance;
+	/** List of UI commands for this manager. This will be filled by this and corresponding classes. */
+	TSharedRef<FUICommandList> CommandList;
+
+	/** An instance of the Object Profiler action manager. */
+	FObjectProfilerActionManager ActionManager;
+
+	/** A weak pointer to the Object Insights window. */
+	TWeakPtr<SObjectProfilerWindow> ProfilerWindowWeakPtr;
+
+	/** If the Objects table/tree view is visible or hidden. */
+	bool bIsObjectTableTreeViewVisible = false;
+
+	/** If the Object Details view is visible or hidden. */
+	bool bIsObjectDetailsViewVisible = false;
 
-	TWeakPtr<FTabManager> TimingTabManager;
+	/** If the Segmented Bar Graph is visible or hidden. */
+	bool bIsSegmentedBarGraphVisible = false;
 
-	TSharedPtr<SMemoryView> MemoryView;
+	/** The id of the Object Insights major tab. */
+	FName MajorTabId;
+
+	/** The name of the Object Profiler log listing. */
+	FName LogListingName;
+
+	/** A shared pointer to the global instance of the Object Profiler manager. */
+	static TSharedPtr<FObjectProfilerManager> Instance;
 };
 
 } // namespace UE::Insights::ObjectProfiler
diff --git a/Engine/Source/Developer/TraceInsights/Private/Insights/ObjectProfiler/ViewModels/AssetInfoNode.h b/Engine/Source/Developer/TraceInsights/Private/Insights/ObjectProfiler/ViewModels/AssetInfoNode.h
--- a/Engine/Source/Developer/TraceInsights/Private/Insights/ObjectProfiler/ViewModels/AssetInfoNode.h
+++ b/Engine/Source/Developer/TraceInsights/Private/Insights/ObjectProfiler/ViewModels/AssetInfoNode.h
@@ -16,8 +16,9 @@ namespace UE::Insights::ObjectProfiler
  */
 struct FAssetInfoNode
 {
-	FSoftObjectPath ActorPath;
 	FAssetData AssetData;
+	FName AssetClassName;
+	FSoftObjectPath ActorPath;
 };
 
 ////////////////////////////////////////////////////////////////////////////////////////////////////
diff --git a/Engine/Source/Developer/TraceInsights/Private/Insights/ObjectProfiler/ViewModels/ObjectNode.cpp b/Engine/Source/Developer/TraceInsights/Private/Insights/ObjectProfiler/ViewModels/ObjectNode.cpp
--- a/Engine/Source/Developer/TraceInsights/Private/Insights/ObjectProfiler/ViewModels/ObjectNode.cpp
+++ b/Engine/Source/Developer/TraceInsights/Private/Insights/ObjectProfiler/ViewModels/ObjectNode.cpp
@@ -140,6 +140,27 @@ TSharedPtr<FObjectNode> FObjectNode::GetPackage() const
 
 ////////////////////////////////////////////////////////////////////////////////////////////////////
 
+void FObjectNode::InitObjectPath()
+{
+	TArray<const TCHAR*> Names;
+	for (FObjectNode* Obj = this; Obj != nullptr; Obj = Obj->GetOuter().Get())
+	{
+		Names.Add(Obj->GetObjectName());
+	}
+	TStringBuilder<1024> Str;
+	for (int32 Index = Names.Num() - 1; Index >= 0; --Index)
+	{
+		Str.Append(Names[Index]);
+		if (Index != 0)
+		{
+			Str.AppendChar('/');
+		}
+	}
+	ObjectPath = Str.ToView();
+}
+
+////////////////////////////////////////////////////////////////////////////////////////////////////
+
 // static
 int64 FObjectNode::GetEstimatedMemorySize(const FBaseTreeNode& Node)
 {
@@ -186,6 +207,34 @@ double FObjectNode::GetEstimatedMemoryImpact() const
 
 ////////////////////////////////////////////////////////////////////////////////////////////////////
 
+FAssetData FObjectNode::MakeAssetData()
+{
+	FString LongPackageName;
+	TSharedPtr<FObjectNode> Package = GetPackage();
+	if (Package.IsValid())
+	{
+		LongPackageName = Package->GetObjectPath();
+	}
+
+	FName AssetClassPackageName;
+	FName AssetClassName;
+	TSharedPtr<FObjectNode> ClassObject = GetClass();
+	if (ClassObject.IsValid())
+	{
+		TSharedPtr<FObjectNode> ClassPackage = ClassObject->GetPackage();
+		if (ClassPackage)
+		{
+			AssetClassPackageName = ClassPackage->GetObjectName();
+		}
+		AssetClassName = ClassObject->GetObjectName();
+	}
+	FTopLevelAssetPath AssetClassPathName{ AssetClassPackageName , AssetClassName };
+
+	return FAssetData(LongPackageName, ObjectPath, AssetClassPathName);
+}
+
+////////////////////////////////////////////////////////////////////////////////////////////////////
+
 //static
 TSharedRef<FObjectNode> FObjectNode::Clone(FObjectNode& InNode)
 {
diff --git a/Engine/Source/Developer/TraceInsights/Private/Insights/ObjectProfiler/ViewModels/ObjectNode.h b/Engine/Source/Developer/TraceInsights/Private/Insights/ObjectProfiler/ViewModels/ObjectNode.h
--- a/Engine/Source/Developer/TraceInsights/Private/Insights/ObjectProfiler/ViewModels/ObjectNode.h
+++ b/Engine/Source/Developer/TraceInsights/Private/Insights/ObjectProfiler/ViewModels/ObjectNode.h
@@ -4,6 +4,8 @@
 
 #include "CoreTypes.h"
 
+#include "AssetRegistry/AssetData.h"
+
 // TraceServices
 #include "TraceServices/Model/ObjectProvider.h"
 
@@ -73,7 +75,10 @@ public:
 	TSharedPtr<FObjectNode> GetPackage() const;
 
 	const TCHAR* GetObjectName() const { return ObjectName; }
-	const TCHAR* GetObjectPath() const { return *ObjectPath; }
+
+	void InitObjectPath();
+	const FString& GetObjectPath() const { return ObjectPath; }
+
 	const TCHAR* GetVersePath() const { return VersePath; }
 
 	virtual int32 GetStructureSize() const { return 0; }
@@ -99,6 +104,8 @@ public:
 	virtual bool IsFunction() const { return false; }
 	virtual bool IsPackage() const { return false; }
 
+	FAssetData MakeAssetData();
+
 	static TSharedRef<FObjectNode> Clone(FObjectNode& InNode);
 
 protected:
diff --git a/Engine/Source/Developer/TraceInsights/Private/Insights/ObjectProfiler/ViewModels/ObjectTable.cpp b/Engine/Source/Developer/TraceInsights/Private/Insights/ObjectProfiler/ViewModels/ObjectTable.cpp
--- a/Engine/Source/Developer/TraceInsights/Private/Insights/ObjectProfiler/ViewModels/ObjectTable.cpp
+++ b/Engine/Source/Developer/TraceInsights/Private/Insights/ObjectProfiler/ViewModels/ObjectTable.cpp
@@ -787,7 +787,7 @@ void FObjectTable::AddDefaultColumns()
 				if (Node.Is<FObjectNode>())
 				{
 					const FObjectNode& ObjectNode = Node.As<FObjectNode>();
-					return FTableCellValue(ObjectNode.GetObjectPath());
+					return FTableCellValue(*ObjectNode.GetObjectPath());
 				}
 				return TOptional<FTableCellValue>();
 			}
diff --git a/Engine/Source/Developer/TraceInsights/Private/Insights/ObjectProfiler/Widgets/SMemoryView.cpp b/Engine/Source/Developer/TraceInsights/Private/Insights/ObjectProfiler/Widgets/SMemoryView.cpp
--- a/Engine/Source/Developer/TraceInsights/Private/Insights/ObjectProfiler/Widgets/SMemoryView.cpp
+++ b/Engine/Source/Developer/TraceInsights/Private/Insights/ObjectProfiler/Widgets/SMemoryView.cpp
@@ -4,6 +4,9 @@
 
 #include "Widgets/Layout/SSplitter.h"
 
+// TraceInsightsCore
+#include "InsightsCore/Widgets/SSegmentedBarGraph.h"
+
 // TraceInsights
 #include "Insights/ObjectProfiler/Widgets/SObjectDetailsView.h"
 #include "Insights/ObjectProfiler/Widgets/SObjectTableTreeView.h"
@@ -29,36 +32,71 @@ SMemoryView::~SMemoryView()
 
 void SMemoryView::Construct(const FArguments& InArgs)
 {
-	SAssignNew(ObjectDetailsView, SObjectDetailsView, SharedThis(this));
-	SAssignNew(ObjectTableTreeView, SObjectTableTreeView, ObjectDetailsView);
-
 	ChildSlot
 	[
-		SNew(SSplitter)
-		.Orientation(Orient_Horizontal)
+		SNew(SVerticalBox)
+
+		+ SVerticalBox::Slot()
+		.Padding(FMargin(4.0f, 4.0f, 4.0f, 0.0f))
+		[
+			SNew(STextBlock)
+			.Text(LOCTEXT("SegmentedBarGraphTitle", "Estimated Memory Breakdown"))
+		]
 
-		+ SSplitter::Slot()
-		.Value(0.6f)
+		+ SVerticalBox::Slot()
+		.Padding(FMargin(4.0f, 2.0f, 4.0f, 4.0f))
 		[
-			ObjectTableTreeView.ToSharedRef()
+			SNew(SBox)
+			.HeightOverride(22.0f)
+			[
+				SNew(SBorder)
+				.BorderBackgroundColor(FLinearColor(0.7f, 0.3f, 0.3f, 1.0f))
+				[
+					SAssignNew(SegmentedBarGraph, SSegmentedBarGraph)
+				]
+			]
 		]
 
-		+ SSplitter::Slot()
-		.Value(0.4f)
+		+ SVerticalBox::Slot()
 		[
-			ObjectDetailsView.ToSharedRef()
+			SNew(SSplitter)
+			.Orientation(Orient_Horizontal)
+
+			+ SSplitter::Slot()
+			.Value(0.6f)
+			[
+				SAssignNew(ObjectTableTreeView, SObjectTableTreeView)
+			]
+
+			+ SSplitter::Slot()
+			.Value(0.4f)
+			[
+				SAssignNew(ObjectDetailsView, SObjectDetailsView)
+			]
 		]
 	];
+
+	ObjectTableTreeView->SetObjectDetailsView(ObjectDetailsView.ToWeakPtr());
+	ObjectTableTreeView->SetSegmentedBarGraph(SegmentedBarGraph.ToWeakPtr());
+	ObjectDetailsView->SetParentMemoryView(SharedThis(this));
+	ObjectDetailsView->SetAssetInfoProvider(AssetInfoProvider);
 }
 
 ////////////////////////////////////////////////////////////////////////////////////////////////////
 
 void SMemoryView::OnClose()
 {
+	if (ObjectDetailsView.IsValid())
+	{
+		//ObjectDetailsView->OnClose();
+		ObjectDetailsView.Reset();
+	}
 	if (ObjectTableTreeView.IsValid())
 	{
 		ObjectTableTreeView->OnClose();
+		ObjectTableTreeView.Reset();
 	}
+	SegmentedBarGraph.Reset();
 }
 
 ////////////////////////////////////////////////////////////////////////////////////////////////////
diff --git a/Engine/Source/Developer/TraceInsights/Private/Insights/ObjectProfiler/Widgets/SMemoryView.h b/Engine/Source/Developer/TraceInsights/Private/Insights/ObjectProfiler/Widgets/SMemoryView.h
--- a/Engine/Source/Developer/TraceInsights/Private/Insights/ObjectProfiler/Widgets/SMemoryView.h
+++ b/Engine/Source/Developer/TraceInsights/Private/Insights/ObjectProfiler/Widgets/SMemoryView.h
@@ -11,9 +11,15 @@
 #include "Insights/ObjectProfiler/ViewModels/AssetInfoNode.h"
 #include "Insights/Widgets/IObjectProfilerMemoryView.h"
 
+namespace UE::Insights
+{
+	class SSegmentedBarGraph;
+}
+
 namespace UE::Insights::ObjectProfiler
 {
 
+class IAssetInfoProvider;
 class SObjectDetailsView;
 class SObjectTableTreeView;
 
@@ -45,15 +51,22 @@ public:
 
 	virtual TSharedRef<SWidget> GetWidget() override { return SharedThis(this); }
 
-	virtual void SetGetActorsDelegate(FGetActors InDelegate) { GetActorsDelegate = MoveTemp(InDelegate); };
-	FGetActors GetActorsDelegate;
+	virtual void SetAssetInfoProvider(TSharedPtr<IAssetInfoProvider> InAssetInfoProvider) override
+	{
+		AssetInfoProvider = InAssetInfoProvider;
+	}
 
 private:
-	/** The object tree view displayed in the left panel. */
+	TSharedPtr<IAssetInfoProvider> AssetInfoProvider;
+
+	/** The Objects tree view displayed in the left panel. */
 	TSharedPtr<SObjectTableTreeView> ObjectTableTreeView;
 
-	/** The object details view displayed in the right panel. */
+	/** The Object Details view displayed in the right panel. */
 	TSharedPtr<SObjectDetailsView> ObjectDetailsView;
+
+	/** The Segmented Bar graph displayed in the toolbar. */
+	TSharedPtr<SSegmentedBarGraph> SegmentedBarGraph;
 };
 
 ////////////////////////////////////////////////////////////////////////////////////////////////////
diff --git a/Engine/Source/Developer/TraceInsights/Private/Insights/ObjectProfiler/Widgets/SObjectDetailsView.cpp b/Engine/Source/Developer/TraceInsights/Private/Insights/ObjectProfiler/Widgets/SObjectDetailsView.cpp
--- a/Engine/Source/Developer/TraceInsights/Private/Insights/ObjectProfiler/Widgets/SObjectDetailsView.cpp
+++ b/Engine/Source/Developer/TraceInsights/Private/Insights/ObjectProfiler/Widgets/SObjectDetailsView.cpp
@@ -4,19 +4,11 @@
 
 #include "Brushes/SlateNoResource.h"
 #include "Modules/ModuleManager.h"
-#if WITH_EDITOR
-#include "ContentBrowserModule.h"
-#include "IContentBrowserSingleton.h"
-#include "UnrealEdGlobals.h"
-#include "Editor/EditorEngine.h"
-#include "Editor/UnrealEdEngine.h"
 
-extern UNREALED_API class UEditorEngine* GEditor;
-
-#endif
 #include "SlateOptMacros.h"
 #include "Styling/AppStyle.h"
 #include "Styling/SlateTypes.h"
+#include "Styling/StyleColors.h"
 #include "Widgets/Images/SImage.h"
 #include "Widgets/Input/SButton.h"
 #include "Widgets/Layout/SScrollBar.h"
@@ -26,8 +18,11 @@ extern UNREALED_API class UEditorEngine* GEditor;
 #include "Widgets/Views/STableRow.h"
 
 // TraceInsights
+#include "Insights/ObjectProfiler/IAssetInfoProvider.h"
 #include "Insights/ObjectProfiler/ViewModels/AssetInfoNode.h"
-#include "SMemoryView.h"
+#include "Insights/ObjectProfiler/ViewModels/ObjectNode.h"
+#include "Insights/ObjectProfiler/Widgets/SMemoryView.h"
+#include "Insights/ObjectProfiler/Widgets/SObjectProfilerWindow.h"
 
 #define LOCTEXT_NAMESPACE "UE::Insights::ObjectProfiler::SObjectDetailsView"
 
@@ -135,16 +130,15 @@ private:
 
 	FReply OnBrowseToAssetClicked()
 	{
-#if WITH_EDITOR
-		TSharedPtr<FAssetInfoNode> Item = WeakItem.Pin();
-		if (Item.IsValid() && Item->AssetData.IsValid())
+		auto AssetInfoProvider = WeakAssetInfoProvider.Pin();
+		if (AssetInfoProvider)
 		{
-			if (FContentBrowserModule* ContentBrowserModule = FModuleManager::GetModulePtr<FContentBrowserModule>("ContentBrowser"))
+			TSharedPtr<FAssetInfoNode> Item = WeakItem.Pin();
+			if (Item)
 			{
-				ContentBrowserModule->Get().SyncBrowserToAssets({ Item->AssetData });
+				AssetInfoProvider->BrowseToAsset(*Item);
 			}
 		}
-#endif
 		return FReply::Handled();
 	}
 
@@ -156,31 +150,21 @@ private:
 
 	FReply OnBrowseToActorClicked()
 	{
-#if WITH_EDITOR
-		TSharedPtr<FAssetInfoNode> Item = WeakItem.Pin();
-		if (Item.IsValid() && Item->ActorPath.IsValid())
+		auto AssetInfoProvider = WeakAssetInfoProvider.Pin();
+		if (AssetInfoProvider)
 		{
-			UObject* Obj = Item->ActorPath.ResolveObject();
-			if (!Obj)
+			TSharedPtr<FAssetInfoNode> Item = WeakItem.Pin();
+			if (Item)
 			{
-				Obj = Item->ActorPath.TryLoad();
-			}
-			if (AActor* AsActor = Cast<AActor>(Obj))
-			{
-				GEditor->SelectNone(false, true);
-				GEditor->SelectActor(AsActor, /*InSelected=*/true, /*bNotify=*/false, /*bSelectEvenIfHidden=*/true);
-				GEditor->NoteSelectionChange();
-				GEditor->MoveViewportCamerasToActor(*AsActor, false);
-				GUnrealEd->ShowActorProperties();
-				GUnrealEd->UpdateFloatingPropertyWindows();
+				AssetInfoProvider->BrowseToActor(*Item);
 			}
 		}
-#endif
 		return FReply::Handled();
 	}
 
 private:
 	TWeakPtr<FAssetInfoNode> WeakItem;
+	TWeakPtr<IAssetInfoProvider> WeakAssetInfoProvider;
 };
 
 ////////////////////////////////////////////////////////////////////////////////////////////////////
@@ -200,10 +184,8 @@ SObjectDetailsView::~SObjectDetailsView()
 ////////////////////////////////////////////////////////////////////////////////////////////////////
 
 BEGIN_SLATE_FUNCTION_BUILD_OPTIMIZATION
-void SObjectDetailsView::Construct(const FArguments& InArgs, TWeakPtr<SMemoryView> MemoryView)
+void SObjectDetailsView::Construct(const FArguments& InArgs)
 {
-	ParentMemoryView = MemoryView;
-
 	SAssignNew(ExternalScrollbar, SScrollBar)
 	.AlwaysShowScrollbar(true);
 
@@ -213,17 +195,83 @@ void SObjectDetailsView::Construct(const FArguments& InArgs, TWeakPtr<SMemoryVie
 
 		+ SVerticalBox::Slot()
 		.AutoHeight()
-		.HAlign(HAlign_Center)
-		.Padding(FMargin(0.0f, 4.0f))
+		.HAlign(HAlign_Left)
+		.Padding(FMargin(8.0f, 8.0f, 8.0f, 4.0f))
+		[
+			SNew(STextBlock)
+			.Font(FAppStyle::Get().GetFontStyle("Font.Large"))
+			.ColorAndOpacity(this, &SObjectDetailsView::GetColorForSelectedObjectName)
+			.Text(this, &SObjectDetailsView::GetSelectedObjectName)
+			.ToolTipText(this, &SObjectDetailsView::GetSelectedObjectToolTip)
+		]
+
+		+ SVerticalBox::Slot()
+		.AutoHeight()
+		.HAlign(HAlign_Left)
+		.Padding(FMargin(8.0f, 0.0f, 8.0f, 4.0f))
+		[
+			SNew(STextBlock)
+			.Text(this, &SObjectDetailsView::GetSelectedObjectSizeText)
+			.ToolTipText(this, &SObjectDetailsView::GetSelectedObjectThumbnailToolTip)
+		]
+
+		+ SVerticalBox::Slot()
+		.AutoHeight()
+		.HAlign(HAlign_Left)
+		.Padding(FMargin(8.0f, 4.0f, 8.0f, 4.0f))
+		[
+			SNew(SBox)
+			.Visibility(this, &SObjectDetailsView::GetSelectedObjectVisibility)
+			.WidthOverride(50.0f)
+			.HeightOverride(50.0f)
+			.ToolTipText(this, &SObjectDetailsView::GetSelectedObjectThumbnailToolTip)
+			[
+				SNew(SBorder)
+				.BorderBackgroundColor(FSlateColor::UseForeground())
+				[
+					SNew(SImage)
+					.Image(this, &SObjectDetailsView::GetSelectedObjectThumbnail)
+				]
+			]
+		]
+
+		+ SVerticalBox::Slot()
+		.AutoHeight()
+		.HAlign(HAlign_Left)
+		.Padding(FMargin(8.0f, 12.0f, 8.0f, 4.0f))
 		[
 			SNew(STextBlock)
-			.Text(LOCTEXT("ObjectDetailsTitle", "Object Details"))
+			.Visibility(this, &SObjectDetailsView::GetValidationIssuesVisibility)
+			//.Font(FAppStyle::Get().GetFontStyle("NormalFontBold"))
+			.Text(LOCTEXT("ValidationIssues", "Issues"))
+		]
+
+		+ SVerticalBox::Slot()
+		.AutoHeight()
+		.HAlign(HAlign_Left)
+		.Padding(FMargin(8.0f, 4.0f, 8.0f, 4.0f))
+		[
+			SNew(STextBlock)
+			.Visibility(this, &SObjectDetailsView::GetValidationIssuesVisibility)
+			.Text(LOCTEXT("TODO", "<placeholder>"))
+		]
+
+		+ SVerticalBox::Slot()
+		.AutoHeight()
+		.HAlign(HAlign_Left)
+		.Padding(FMargin(8.0f, 12.0f, 8.0f, 4.0f))
+		[
+			SNew(STextBlock)
+			.Visibility(this, &SObjectDetailsView::GetActorsVisibility)
+			//.Font(FAppStyle::Get().GetFontStyle("NormalFontBold"))
+			.Text(LOCTEXT("RefBy", "Referenced by"))
 		]
 
 		+ SVerticalBox::Slot()
 		.FillHeight(1.0f)
 		[
 			SNew(SHorizontalBox)
+			.Visibility(this, &SObjectDetailsView::GetActorsVisibility)
 
 			+ SHorizontalBox::Slot()
 			.FillWidth(1.0f)
@@ -250,23 +298,39 @@ END_SLATE_FUNCTION_BUILD_OPTIMIZATION
 
 ////////////////////////////////////////////////////////////////////////////////////////////////////
 
-void SObjectDetailsView::SetMatchedAssets(const FAssetData& AssetData)
+TSharedRef<ITableRow> SObjectDetailsView::OnGenerateRow(TSharedPtr<FAssetInfoNode> InItem, const TSharedRef<STableViewBase>& OwnerTable)
 {
+	return SNew(SObjectDetailsViewRow, InItem, OwnerTable);
+}
+
+////////////////////////////////////////////////////////////////////////////////////////////////////
+
+void SObjectDetailsView::SetObject(TSharedPtr<const FObjectNode> InObjectNode)
+{
+	SelectedObjectNode = InObjectNode;
+	SelectedAssetData = FAssetData();
+	SelectedAssetClassName = NAME_None;
+
 	AssetEntries.Empty();
 
-	if (AssetData.IsValid())
+	if (InObjectNode.IsValid())
 	{
-		auto MemoryView = ParentMemoryView.Pin();
-		if (MemoryView)
+		auto AssetInfoProvider = WeakAssetInfoProvider.Pin();
+		if (AssetInfoProvider)
 		{
-			if (MemoryView->GetActorsDelegate.IsBound())
+			AssetInfoProvider->GetAssetData(*InObjectNode.Get(), SelectedAssetData);
+			SelectedAssetClassName = FName(InObjectNode->GetClass()->GetObjectName());
+
+			if (SelectedAssetData.IsValid())
 			{
 				TArray<FSoftObjectPath> Actors;
-				MemoryView->GetActorsDelegate.Execute(AssetData, Actors);
+				AssetInfoProvider->GetActors(SelectedAssetData, Actors);
 
-				for (auto Actor : Actors)
+				for (const auto& Actor : Actors)
 				{
 					auto Entry = MakeShared<FAssetInfoNode>();
+					Entry->AssetData = SelectedAssetData;
+					Entry->AssetClassName = SelectedAssetClassName;
 					Entry->ActorPath = Actor;
 					AssetEntries.Add(Entry);
 				}
@@ -279,9 +343,115 @@ void SObjectDetailsView::SetMatchedAssets(const FAssetData& AssetData)
 
 ////////////////////////////////////////////////////////////////////////////////////////////////////
 
-TSharedRef<ITableRow> SObjectDetailsView::OnGenerateRow(TSharedPtr<FAssetInfoNode> InItem, const TSharedRef<STableViewBase>& OwnerTable)
+EVisibility SObjectDetailsView::GetSelectedObjectVisibility() const
 {
-	return SNew(SObjectDetailsViewRow, InItem, OwnerTable);
+	return SelectedObjectNode.IsValid() ? EVisibility::Visible : EVisibility::Collapsed;
+}
+
+////////////////////////////////////////////////////////////////////////////////////////////////////
+
+EVisibility SObjectDetailsView::GetActorsVisibility() const
+{
+	return SelectedObjectNode.IsValid() && AssetEntries.Num() > 0 ? EVisibility::Visible : EVisibility::Collapsed;
+}
+
+////////////////////////////////////////////////////////////////////////////////////////////////////
+
+EVisibility SObjectDetailsView::GetValidationIssuesVisibility() const
+{
+	//return SelectedObjectNode.IsValid() && ValidationIssues.Num() > 0 ? EVisibility::Visible : EVisibility::Collapsed;
+	return EVisibility::Collapsed;
+}
+
+////////////////////////////////////////////////////////////////////////////////////////////////////
+
+FSlateColor SObjectDetailsView::GetColorForSelectedObjectName() const
+{
+	if (SelectedObjectNode.IsValid())
+	{
+		return FSlateColor::UseForeground();
+	}
+	return FSlateColor(EStyleColor::White25);
+}
+
+////////////////////////////////////////////////////////////////////////////////////////////////////
+
+FText SObjectDetailsView::GetSelectedObjectName() const
+{
+	if (SelectedObjectNode.IsValid())
+	{
+		return FText::FromString(SelectedObjectNode->GetObjectName());
+	}
+	return LOCTEXT("NoObjectSelected", "No Object Selected");
+}
+
+////////////////////////////////////////////////////////////////////////////////////////////////////
+
+FText SObjectDetailsView::GetSelectedObjectSizeText() const
+{
+	if (SelectedObjectNode.IsValid())
+	{
+		return FText::Format(LOCTEXT("SizeFmt", "Size: {0}"), FText::AsMemory(SelectedObjectNode->GetEstimatedMemorySize()));
+	}
+	return FText::GetEmpty();
+}
+
+////////////////////////////////////////////////////////////////////////////////////////////////////
+
+FText SObjectDetailsView::GetSelectedObjectToolTip() const
+{
+	if (SelectedObjectNode.IsValid())
+	{
+		return FText::FromString(SelectedObjectNode->GetObjectPath());
+	}
+	return FText::GetEmpty();
+}
+
+////////////////////////////////////////////////////////////////////////////////////////////////////
+
+FText SObjectDetailsView::GetSelectedObjectThumbnailToolTip() const
+{
+	if (SelectedObjectNode.IsValid())
+	{
+		TUtf8StringBuilder<1024> Str;
+		Str.Append(SelectedObjectNode->GetObjectName());
+		Str.Append("\nId: ");
+		Str.Appendf("%u", SelectedObjectNode->GetObjectId());
+		Str.Append("\nClass: ");
+		Str.Append(SelectedObjectNode->GetClass()->GetObjectName());
+		if (SelectedObjectNode->GetOuter().IsValid())
+		{
+			Str.Append("\nOuter: ");
+			Str.Append(SelectedObjectNode->GetOuter()->GetObjectName());
+		}
+		if (SelectedObjectNode->GetSystemMemorySize() > 0)
+		{
+			Str.Appendf("\nEstimated System Memory: %.2f MiB", double(SelectedObjectNode->GetSystemMemorySize()) / (1024.0 * 1024.0));
+		}
+		if (SelectedObjectNode->GetVideoMemorySize() > 0)
+		{
+			Str.Appendf("\nEstimated Video Memory: %.2f MiB", double(SelectedObjectNode->GetVideoMemorySize()) / (1024.0 * 1024.0));
+		}
+		return FText::FromString(Str.ToString());
+	}
+	return FText::GetEmpty();
+}
+
+////////////////////////////////////////////////////////////////////////////////////////////////////
+
+const FSlateBrush* SObjectDetailsView::GetSelectedObjectThumbnail() const
+{
+	if (SelectedObjectNode.IsValid() &&
+		SelectedAssetData.IsValid() &&
+		SelectedAssetClassName != NAME_None)
+	{
+		auto AssetInfoProvider = WeakAssetInfoProvider.Pin();
+		if (AssetInfoProvider)
+		{
+			return AssetInfoProvider->GetThumbnail(SelectedAssetData, SelectedAssetClassName);
+		}
+	}
+	return nullptr;
 }
 
 ////////////////////////////////////////////////////////////////////////////////////////////////////
diff --git a/Engine/Source/Developer/TraceInsights/Private/Insights/ObjectProfiler/Widgets/SObjectDetailsView.h b/Engine/Source/Developer/TraceInsights/Private/Insights/ObjectProfiler/Widgets/SObjectDetailsView.h
--- a/Engine/Source/Developer/TraceInsights/Private/Insights/ObjectProfiler/Widgets/SObjectDetailsView.h
+++ b/Engine/Source/Developer/TraceInsights/Private/Insights/ObjectProfiler/Widgets/SObjectDetailsView.h
@@ -8,11 +8,17 @@
 #include "Widgets/SCompoundWidget.h"
 #include "Widgets/Views/SListView.h"
 
+// TraceInsights
 #include "Insights/ObjectProfiler/ViewModels/AssetInfoNode.h"
 
 namespace UE::Insights::ObjectProfiler
 {
+
+class IAssetInfoProvider;
+class SObjectProfilerWindow;
 class SMemoryView;
+class FObjectNode;
+
 ////////////////////////////////////////////////////////////////////////////////////////////////////
 
 /**
@@ -30,22 +36,56 @@ public:
 	SLATE_BEGIN_ARGS(SObjectDetailsView) {}
 	SLATE_END_ARGS()
 
-	void Construct(const FArguments& InArgs, TWeakPtr<SMemoryView> MemoryView);
+	void Construct(const FArguments& InArgs);
+
+	void Reset() {}
+	void OnClose() {}
+
+	void SetParentWindow(TSharedPtr<SObjectProfilerWindow> InParentWindow)
+	{
+		ParentWindow = InParentWindow.ToWeakPtr();
+	}
 
-	void SetMatchedAssets(const FAssetData& AssetData);
+	void SetParentMemoryView(TSharedPtr<SMemoryView> InParentMemoryView)
+	{
+		ParentMemoryView = InParentMemoryView.ToWeakPtr();
+	}
+
+	void SetAssetInfoProvider(TSharedPtr<IAssetInfoProvider> InAssetInfoProvider)
+	{
+		WeakAssetInfoProvider = InAssetInfoProvider.ToWeakPtr();
+	}
+
+	void SetObject(TSharedPtr<const FObjectNode> InObjectNode);
 
 protected:
 	TSharedRef<ITableRow> OnGenerateRow(TSharedPtr<FAssetInfoNode> InItem, const TSharedRef<STableViewBase>& OwnerTable);
 
+	EVisibility GetSelectedObjectVisibility() const;
+	EVisibility GetValidationIssuesVisibility() const;
+	EVisibility GetActorsVisibility() const;
+	FSlateColor GetColorForSelectedObjectName() const;
+	FText GetSelectedObjectName() const;
+	FText GetSelectedObjectSizeText() const;
+	FText GetSelectedObjectToolTip() const;
+	FText GetSelectedObjectThumbnailToolTip() const;
+	const FSlateBrush* GetSelectedObjectThumbnail() const;
+
 private:
-	TSharedPtr<SListView<TSharedPtr<FAssetInfoNode>>> ListView;
+	TWeakPtr<IAssetInfoProvider> WeakAssetInfoProvider;
 
-	TSharedPtr<SScrollBar> ExternalScrollbar;
+	TWeakPtr<SObjectProfilerWindow> ParentWindow;
+	TWeakPtr<SMemoryView> ParentMemoryView;
+
+	TSharedPtr<const FObjectNode> SelectedObjectNode;
+	FAssetData SelectedAssetData;
+	FName SelectedAssetClassName;
 
 	/** List of matched asset info nodes shown in the list view. */
 	TArray<TSharedPtr<FAssetInfoNode>> AssetEntries;
 
-	TWeakPtr<SMemoryView> ParentMemoryView;
+	TSharedPtr<SListView<TSharedPtr<FAssetInfoNode>>> ListView;
+	TSharedPtr<SScrollBar> ExternalScrollbar;
 };
 
 ////////////////////////////////////////////////////////////////////////////////////////////////////
diff --git a/Engine/Source/Developer/TraceInsights/Private/Insights/ObjectProfiler/Widgets/SObjectProfilerWindow.cpp b/Engine/Source/Developer/TraceInsights/Private/Insights/ObjectProfiler/Widgets/SObjectProfilerWindow.cpp
new file mode 100644
--- /dev/null
+++ b/Engine/Source/Developer/TraceInsights/Private/Insights/ObjectProfiler/Widgets/SObjectProfilerWindow.cpp
@@ -0,0 +1,315 @@
+// Copyright Epic Games, Inc. All Rights Reserved.
+
+#include "SObjectProfilerWindow.h"
+
+#include "Modules/ModuleManager.h"
+#include "SlateOptMacros.h"
+#include "Widgets/Docking/SDockTab.h"
+#include "Widgets/Layout/SSplitter.h"
+
+// TraceInsightsCore
+#include "InsightsCore/Widgets/SSegmentedBarGraph.h"
+
+// TraceInsights
+#include "Insights/InsightsStyle.h"
+#include "Insights/IUnrealInsightsModule.h"
+#include "Insights/ObjectProfiler/ObjectProfilerManager.h"
+#include "Insights/ObjectProfiler/Widgets/SObjectDetailsView.h"
+#include "Insights/ObjectProfiler/Widgets/SObjectTableTreeView.h"
+
+#define LOCTEXT_NAMESPACE "UE::Insights::ObjectProfiler"
+
+namespace UE::Insights::ObjectProfiler
+{
+
+const FName FObjectProfilerTabs::ObjectTableTreeViewId("ObjectTableTreeView");
+const FName FObjectProfilerTabs::ObjectDetailsViewId("ObjectDetailsView");
+const FName SObjectProfilerWindow::ObjectProfilerTabId("ObjectProfiler");
+
+////////////////////////////////////////////////////////////////////////////////////////////////////
+
+SObjectProfilerWindow::SObjectProfilerWindow()
+	: SMajorTabWindow(ObjectProfilerTabId)
+{
+}
+
+////////////////////////////////////////////////////////////////////////////////////////////////////
+
+SObjectProfilerWindow::~SObjectProfilerWindow()
+{
+	CloseAllOpenTabs();
+
+	check(ObjectTableTreeView == nullptr);
+	check(ObjectDetailsView == nullptr);
+}
+
+////////////////////////////////////////////////////////////////////////////////////////////////////
+
+const TCHAR* SObjectProfilerWindow::GetAnalyticsEventName() const
+{
+	return TEXT("Insights.Usage.ObjectProfiler");
+}
+
+////////////////////////////////////////////////////////////////////////////////////////////////////
+
+void SObjectProfilerWindow::Reset()
+{
+	if (ObjectTableTreeView)
+	{
+		ObjectTableTreeView->Reset();
+	}
+
+	if (ObjectDetailsView)
+	{
+		ObjectDetailsView->Reset();
+	}
+}
+
+////////////////////////////////////////////////////////////////////////////////////////////////////
+
+BEGIN_SLATE_FUNCTION_BUILD_OPTIMIZATION
+TSharedRef<SDockTab> SObjectProfilerWindow::SpawnTab_ObjectTableTreeView(const FSpawnTabArgs& Args)
+{
+	FObjectProfilerManager::Get()->SetObjectTableTreeViewVisible(true);
+
+	check(ObjectTableTreeView == nullptr);
+	const TSharedRef<SDockTab> DockTab = SNew(SDockTab)
+		.ShouldAutosize(false)
+		.TabRole(ETabRole::PanelTab)
+		[
+			SAssignNew(ObjectTableTreeView, SObjectTableTreeView)
+		];
+	check(ObjectTableTreeView != nullptr);
+
+	ObjectTableTreeView->SetLogListingName(FObjectProfilerManager::Get()->GetLogListingName());
+
+	if (ObjectDetailsView)
+	{
+		ObjectTableTreeView->SetObjectDetailsView(ObjectDetailsView);
+	}
+
+	//SharedState->BindCommands();
+
+	DockTab->SetOnTabClosed(SDockTab::FOnTabClosedCallback::CreateRaw(this, &SObjectProfilerWindow::OnTabClosed_ObjectTableTreeView));
+	AddOpenTab(DockTab);
+
+	return DockTab;
+}
+END_SLATE_FUNCTION_BUILD_OPTIMIZATION
+
+////////////////////////////////////////////////////////////////////////////////////////////////////
+
+void SObjectProfilerWindow::OnTabClosed_ObjectTableTreeView(TSharedRef<SDockTab> TabBeingClosed)
+{
+	FObjectProfilerManager::Get()->SetObjectTableTreeViewVisible(false);
+
+	if (ObjectTableTreeView)
+	{
+		ObjectTableTreeView = nullptr;
+	}
+
+	RemoveOpenTab(TabBeingClosed);
+}
+
+////////////////////////////////////////////////////////////////////////////////////////////////////
+
+BEGIN_SLATE_FUNCTION_BUILD_OPTIMIZATION
+TSharedRef<SDockTab> SObjectProfilerWindow::SpawnTab_ObjectDetailsView(const FSpawnTabArgs& Args)
+{
+	FObjectProfilerManager::Get()->SetObjectDetailsViewVisible(true);
+
+	check(ObjectDetailsView == nullptr);
+	const TSharedRef<SDockTab> DockTab = SNew(SDockTab)
+		.ShouldAutosize(false)
+		.TabRole(ETabRole::PanelTab)
+		[
+			SAssignNew(ObjectDetailsView, SObjectDetailsView)
+		];
+	check(ObjectDetailsView != nullptr);
+
+	ObjectDetailsView->SetAssetInfoProvider(AssetInfoProvider);
+
+	if (ObjectTableTreeView)
+	{
+		ObjectTableTreeView->SetObjectDetailsView(ObjectDetailsView);
+	}
+
+	DockTab->SetOnTabClosed(SDockTab::FOnTabClosedCallback::CreateRaw(this, &SObjectProfilerWindow::OnTabClosed_ObjectDetailsView));
+	AddOpenTab(DockTab);
+
+	return DockTab;
+}
+END_SLATE_FUNCTION_BUILD_OPTIMIZATION
+
+////////////////////////////////////////////////////////////////////////////////////////////////////
+
+void SObjectProfilerWindow::OnTabClosed_ObjectDetailsView(TSharedRef<SDockTab> TabBeingClosed)
+{
+	FObjectProfilerManager::Get()->SetObjectDetailsViewVisible(false);
+
+	if (ObjectDetailsView)
+	{
+		ObjectDetailsView = nullptr;
+	}
+	if (ObjectTableTreeView)
+	{
+		ObjectTableTreeView->SetObjectDetailsView(nullptr);
+	}
+
+	RemoveOpenTab(TabBeingClosed);
+}
+
+////////////////////////////////////////////////////////////////////////////////////////////////////
+
+void SObjectProfilerWindow::Construct(const FArguments& InArgs, const TSharedRef<SDockTab>& ConstructUnderMajorTab, const TSharedPtr<SWindow>& ConstructUnderWindow)
+{
+	TSharedPtr<FObjectProfilerManager> ObjectProfilerManager = FObjectProfilerManager::Get();
+	ensure(ObjectProfilerManager.IsValid());
+
+	SetCommandList(ObjectProfilerManager->GetCommandList());
+
+	SMajorTabWindow::FArguments Args;
+	SMajorTabWindow::Construct(Args, ConstructUnderMajorTab, ConstructUnderWindow);
+}
+
+////////////////////////////////////////////////////////////////////////////////////////////////////
+
+TSharedRef<FWorkspaceItem> SObjectProfilerWindow::CreateWorkspaceMenuGroup()
+{
+	return GetTabManager()->AddLocalWorkspaceMenuCategory(LOCTEXT("ObjectProfilerMenuGroupName", "Object Insights"));
+}
+
+////////////////////////////////////////////////////////////////////////////////////////////////////
+
+void SObjectProfilerWindow::RegisterTabSpawners()
+{
+	auto TabManagerPtr = GetTabManager();
+	check(TabManagerPtr.IsValid());
+	FTabManager& TabManagerRef = *TabManagerPtr.Get();
+
+	check(GetWorkspaceMenuGroup().IsValid());
+	const TSharedRef<FWorkspaceItem> Group = GetWorkspaceMenuGroup().ToSharedRef();
+
+	IUnrealInsightsModule& InsightsModule = FModuleManager::LoadModuleChecked<IUnrealInsightsModule>("TraceInsights");
+	const FInsightsMajorTabConfig& Config = InsightsModule.FindMajorTabConfig(ObjectProfilerTabId);
+
+	if (Config.ShouldRegisterMinorTab(FObjectProfilerTabs::ObjectTableTreeViewId))
+	{
+		TabManagerRef.RegisterTabSpawner(FObjectProfilerTabs::ObjectTableTreeViewId,
+			FOnSpawnTab::CreateRaw(this, &SObjectProfilerWindow::SpawnTab_ObjectTableTreeView))
+				.SetDisplayName(LOCTEXT("ObjectTableTreeViewTabTitle", "Objects"))
+				.SetIcon(FSlateIcon(FInsightsStyle::GetStyleSetName(), "Icons.ObjectTableTreeView"))
+				.SetTooltipText(LOCTEXT("ObjectTableTreeViewTabToolTip", "Opens the Objects table/tree view."))
+				.SetGroup(Group);
+	}
+
+	if (Config.ShouldRegisterMinorTab(FObjectProfilerTabs::ObjectDetailsViewId))
+	{
+		TabManagerRef.RegisterTabSpawner(FObjectProfilerTabs::ObjectDetailsViewId,
+			FOnSpawnTab::CreateRaw(this, &SObjectProfilerWindow::SpawnTab_ObjectDetailsView))
+				.SetDisplayName(LOCTEXT("ObjectDetailsViewTabTitle", "Object Details"))
+				.SetIcon(FSlateIcon(FInsightsStyle::GetStyleSetName(), "Icons.ObjectDetailsView"))
+				.SetTooltipText(LOCTEXT("ObjectDetailsViewTabToolTip", "Opens the Object Details panel."))
+				.SetGroup(Group);
+	}
+}
+
+////////////////////////////////////////////////////////////////////////////////////////////////////
+
+BEGIN_SLATE_FUNCTION_BUILD_OPTIMIZATION
+TSharedRef<FTabManager::FLayout> SObjectProfilerWindow::CreateDefaultTabLayout() const
+{
+	return FTabManager::NewLayout("InsightsObjectProfilerLayout_v1.0")
+		->AddArea
+		(
+			FTabManager::NewPrimaryArea()
+			->SetOrientation(Orient_Horizontal)
+			->Split
+			(
+				FTabManager::NewStack()
+				->SetSizeCoefficient(0.65f)
+				->SetHideTabWell(true)
+				->AddTab(FObjectProfilerTabs::ObjectTableTreeViewId, ETabState::OpenedTab)
+			)
+			->Split
+			(
+				FTabManager::NewStack()
+				->SetSizeCoefficient(0.35f)
+				->SetHideTabWell(true)
+				->AddTab(FObjectProfilerTabs::ObjectDetailsViewId, ETabState::OpenedTab)
+			)
+		);
+}
+END_SLATE_FUNCTION_BUILD_OPTIMIZATION
+
+////////////////////////////////////////////////////////////////////////////////////////////////////
+
+TSharedRef<SWidget> SObjectProfilerWindow::CreateToolbar(TSharedPtr<FExtender> Extender)
+{
+	//IUnrealInsightsModule& InsightsModule = FModuleManager::LoadModuleChecked<IUnrealInsightsModule>("TraceInsights");
+	//const FInsightsMajorTabConfig& Config = InsightsModule.FindMajorTabConfig(ObjectProfilerTabId);
+
+	//return SNew(SObjectProfilerToolbar, Config).ToolbarExtender(Extender);
+
+#if 0 // debug code
+	auto Segment1 = MakeShared<FCustomBarGraphSegment>();
+	Segment1->Size = 40.0;
+	Segment1->Text = LOCTEXT("Segment1", "Textures");
+	Segment1->ToolTipText = LOCTEXT("Segment1_TT", "All textures");
+	Segment1->Color = FLinearColor(0.5f, 0.0f, 0.0f, 1.0f);
+	Segment1->TextColor = FLinearColor(1.0f, 0.0f, 0.0f, 1.0f);
+	Segments.Add(Segment1);
+
+	auto Segment2 = MakeShared<FCustomBarGraphSegment>();
+	Segment2->Size = 40.0;
+	Segment2->Text = LOCTEXT("Segment2", "Materials");
+	Segment2->ToolTipText = LOCTEXT("Segment2_TT", "All materials");
+	Segment2->Color = FLinearColor(0.0f, 0.3f, 0.0f, 1.0f);
+	Segment2->TextColor = FLinearColor(1.0f, 0.0f, 0.0f, 1.0f);
+	Segments.Add(Segment2);
+
+	auto Segment3 = MakeShared<FCustomBarGraphSegment>();
+	Segment3->Size = 40.0;
+	Segment3->Text = LOCTEXT("Segment3_Verse", "Verse");
+	Segment3->ToolTipText = LOCTEXT("Segment3_TT", "Verse objects");
+	Segment3->Color = FLinearColor(0.0f, 0.0f, 0.5f, 1.0f);
+	Segment3->TextColor = FLinearColor(1.0f, 0.0f, 0.0f, 1.0f);
+	Segments.Add(Segment3);
+#endif
+
+	return SNew(SBorder)
+		.BorderImage(FAppStyle::Get().GetBrush("WhiteBrush"))
+		.BorderBackgroundColor(FLinearColor(0.014f, 0.014f, 0.014f, 1.0f))
+		[
+			SNew(SVerticalBox)
+
+			+ SVerticalBox::Slot()
+			.Padding(FMargin(4.0f, 4.0f, 4.0f, 0.0f))
+			[
+				SNew(STextBlock)
+				.Text(LOCTEXT("SegmentedBarGraphTitle", "Estimated Memory Breakdown"))
+			]
+
+			+ SVerticalBox::Slot()
+			.Padding(FMargin(4.0f, 2.0f, 4.0f, 4.0f))
+			[
+				SNew(SBox)
+				.HeightOverride(22.0f)
+				[
+					SNew(SBorder)
+					.BorderImage(FAppStyle::Get().GetBrush("WhiteBrush"))
+					.BorderBackgroundColor(FLinearColor(0.0f, 0.1f, 0.0f, 1.0f))
+					[
+						SAssignNew(SegmentedBarGraph, SSegmentedBarGraph)
+						//.SegmentsSource(&Segments)
+					]
+				]
+			]
+		];
+}
+
+////////////////////////////////////////////////////////////////////////////////////////////////////
+
+} // namespace UE::Insights::ObjectProfiler
+
+#undef LOCTEXT_NAMESPACE
diff --git a/Engine/Source/Developer/TraceInsights/Private/Insights/ObjectProfiler/Widgets/SObjectProfilerWindow.h b/Engine/Source/Developer/TraceInsights/Private/Insights/ObjectProfiler/Widgets/SObjectProfilerWindow.h
new file mode 100644
--- /dev/null
+++ b/Engine/Source/Developer/TraceInsights/Private/Insights/ObjectProfiler/Widgets/SObjectProfilerWindow.h
@@ -0,0 +1,127 @@
+// Copyright Epic Games, Inc. All Rights Reserved.
+
+#pragma once
+
+#include "CoreTypes.h"
+
+#include "UObject/NameTypes.h"
+
+// TraceInsightsCore
+#include "InsightsCore/ViewModels/IBarGraphSegment.h"
+
+// TraceInsights
+#include "Insights/Widgets/SMajorTabWindow.h"
+#include "Insights/Widgets/IObjectProfilerMemoryView.h"
+
+namespace UE::Insights
+{
+	class SSegmentedBarGraph;
+}
+
+namespace UE::Insights::ObjectProfiler
+{
+
+class SObjectTableTreeView;
+class SObjectDetailsView;
+
+struct FObjectProfilerTabs
+{
+	// Tab identifiers
+	static const FName ObjectTableTreeViewId;
+	static const FName ObjectDetailsViewId;
+};
+
+////////////////////////////////////////////////////////////////////////////////////////////////////
+
+class FCustomBarGraphSegment : public UE::Insights::IBarGraphSegment
+{
+	friend class SObjectProfilerWindow;
+
+public:
+	FCustomBarGraphSegment() {}
+	virtual ~FCustomBarGraphSegment() {}
+
+	virtual double GetSize() const { return Size; }
+	virtual FText GetText() const { return Text; }
+	virtual FText GetToolTipText() const { return ToolTipText; }
+	virtual FLinearColor GetColor() const { return Color; }
+	virtual FLinearColor GetTextColor() const { return TextColor; }
+
+private:
+	double Size;
+	FText Text;
+	FText ToolTipText;
+	FLinearColor Color;
+	FLinearColor TextColor;
+};
+
+////////////////////////////////////////////////////////////////////////////////////////////////////
+
+/** Implements the Object Insights window. */
+class SObjectProfilerWindow : public ::Insights::SMajorTabWindow, public IObjectProfilerMemoryView
+{
+public:
+	static const FName ObjectProfilerTabId;
+
+	/** Default constructor. */
+	SObjectProfilerWindow();
+
+	/** Virtual destructor. */
+	virtual ~SObjectProfilerWindow();
+
+	SLATE_BEGIN_ARGS(SObjectProfilerWindow) {}
+	SLATE_END_ARGS()
+
+	virtual void Reset() override;
+
+	/** Constructs this widget. */
+	void Construct(const FArguments& InArgs, const TSharedRef<SDockTab>& ConstructUnderMajorTab, const TSharedPtr<SWindow>& ConstructUnderWindow);
+
+	virtual void SetAssetInfoProvider(TSharedPtr<IAssetInfoProvider> InAssetInfoProvider) override
+	{
+		AssetInfoProvider = InAssetInfoProvider;
+	}
+
+	TSharedPtr<SObjectTableTreeView> GetObjectTableTreeView() const { return ObjectTableTreeView; }
+	TSharedPtr<SObjectDetailsView> GetObjectDetailsView() const { return ObjectDetailsView; }
+
+	void ShowHideSegmentedBarGraph(bool bOnOff) { /* TODO */ }
+
+	const FName& GetLogListingName() const { return LogListingName; }
+
+	virtual TSharedRef<SWidget> GetWidget() override { return SharedThis(this); }
+
+protected:
+	virtual const TCHAR* GetAnalyticsEventName() const override;
+	virtual TSharedRef<FWorkspaceItem> CreateWorkspaceMenuGroup() override;
+	virtual void RegisterTabSpawners() override;
+	virtual TSharedRef<FTabManager::FLayout> CreateDefaultTabLayout() const override;
+	virtual TSharedRef<SWidget> CreateToolbar(TSharedPtr<FExtender> Extender);
+
+private:
+	TSharedRef<SDockTab> SpawnTab_ObjectTableTreeView(const FSpawnTabArgs& Args);
+	void OnTabClosed_ObjectTableTreeView(TSharedRef<SDockTab> TabBeingClosed);
+
+	TSharedRef<SDockTab> SpawnTab_ObjectDetailsView(const FSpawnTabArgs& Args);
+	void OnTabClosed_ObjectDetailsView(TSharedRef<SDockTab> TabBeingClosed);
+
+private:
+	TSharedPtr<IAssetInfoProvider> AssetInfoProvider;
+
+	/** The Objects table tree view widget */
+	TSharedPtr<SObjectTableTreeView> ObjectTableTreeView;
+
+	/** The Object Details panel widget */
+	TSharedPtr<SObjectDetailsView> ObjectDetailsView;
+
+	/** The name of the Object Insights log listing. */
+	FName LogListingName;
+
+	TSharedPtr<SSegmentedBarGraph> SegmentedBarGraph;
+	bool bIsSegmentedBarGraphVisible = true;
+	TArray<TSharedPtr<FCustomBarGraphSegment>> Segments;
+};
+
+////////////////////////////////////////////////////////////////////////////////////////////////////
+
+} // namespace UE::Insights::ObjectProfiler
diff --git a/Engine/Source/Developer/TraceInsights/Private/Insights/ObjectProfiler/Widgets/SObjectTableTreeView.cpp b/Engine/Source/Developer/TraceInsights/Private/Insights/ObjectProfiler/Widgets/SObjectTableTreeView.cpp
--- a/Engine/Source/Developer/TraceInsights/Private/Insights/ObjectProfiler/Widgets/SObjectTableTreeView.cpp
+++ b/Engine/Source/Developer/TraceInsights/Private/Insights/ObjectProfiler/Widgets/SObjectTableTreeView.cpp
@@ -2,9 +2,6 @@
 
 #include "SObjectTableTreeView.h"
 
-#include "AssetRegistry/AssetRegistryModule.h"
-#include "AssetRegistry/IAssetRegistry.h"
-#include "AssetRegistry/AssetData.h"
 #include "DesktopPlatformModule.h"
 #include "Framework/MultiBox/MultiBoxBuilder.h"
 #include "Misc/ConfigCacheIni.h"
@@ -24,6 +21,7 @@
 // TraceInsightsCore
 #include "InsightsCore/Common/Stopwatch.h"
 #include "InsightsCore/Table/ViewModels/Table.h"
+#include "InsightsCore/Widgets/SSegmentedBarGraph.h"
 
 // TraceInsights
 #include "Insights/InsightsManager.h"
@@ -55,10 +53,8 @@ SObjectTableTreeView::~SObjectTableTreeView()
 
 ////////////////////////////////////////////////////////////////////////////////////////////////////
 
-void SObjectTableTreeView::Construct(const FArguments& InArgs, TWeakPtr<SObjectDetailsView> InDetailsView)
+void SObjectTableTreeView::Construct(const FArguments& InArgs)
 {
-	DetailsView = InDetailsView;
-
 	TSharedRef<FObjectTable> TablePtr = MakeShared<FObjectTable>();
 	TablePtr->Reset();
 	TablePtr->SetDisplayName(LOCTEXT("ObjectsTableName", "Objects"));
@@ -75,233 +71,230 @@ void SObjectTableTreeView::Construct(const FArguments& InArgs, TWeakPtr<SObjectD
 
 ////////////////////////////////////////////////////////////////////////////////////////////////////
 
+BEGIN_SLATE_FUNCTION_BUILD_OPTIMIZATION
 void SObjectTableTreeView::ConstructHeaderArea(TSharedRef<SVerticalBox> InHostBox)
 {
-	TSharedPtr<SWidget> FilterToolbarWidget = ConstructFilterToolbar();
-	if (FilterToolbarWidget.IsValid())
-	{
-		InHostBox->AddSlot()
-			.AutoHeight()
-			.Padding(4.0f, 2.0f, 4.0f, 2.0f)
-			[
-				FilterToolbarWidget.ToSharedRef()
-			];
-	}
+	//////////////////////////////////////////////////
+	// View Presets
 
-	TSharedPtr<SWidget> HierarchyBreadcrumbTrailWidget = ConstructHierarchyBreadcrumbTrail();
-	if (HierarchyBreadcrumbTrailWidget.IsValid())
-	{
-		InHostBox->AddSlot()
-			.AutoHeight()
-			.Padding(4.0f, 2.0f, 4.0f, -2.0f)
-			[
-				HierarchyBreadcrumbTrailWidget.ToSharedRef()
-			];
-	}
-}
+	TSharedRef<SHorizontalBox> TopLine = SNew(SHorizontalBox);
+	TopLine->AddSlot()
+		.FillWidth(1.0f)
+		.Padding(0.0f, 0.0f, 6.0f, 0.0f)
+		.VAlign(VAlign_Center)
+		[
+			SNew(SBox)
+		];
+	ConstructViewPreset(TopLine, 140.0f);
 
-////////////////////////////////////////////////////////////////////////////////////////////////////
+	InHostBox->AddSlot()
+		.AutoHeight()
+		.Padding(4.0f, 2.0f, 4.0f, 4.0f)
+		[
+			TopLine
+		];
 
-void SObjectTableTreeView::ConstructFooterArea(TSharedRef<SVerticalBox> InHostBox)
-{
-	TSharedPtr<SWidget> FooterWidget = ConstructFooter();
-	if (FooterWidget.IsValid())
-	{
-		InHostBox->AddSlot()
-			.AutoHeight()
-			.Padding(0.0f, 6.0f, 0.0f, 0.0f)
+	//////////////////////////////////////////////////
+	// Filters
+
+	InHostBox->AddSlot()
+		.AutoHeight()
+		.Padding(4.0f, 2.0f, 4.0f, 2.0f)
+		[
+			SNew(SHorizontalBox)
+
+			+ SHorizontalBox::Slot()
+			.AutoWidth()
+			.Padding(0.0f, 0.0f, 2.0f, 0.0f)
+			.VAlign(VAlign_Center)
 			[
-				FooterWidget.ToSharedRef()
-			];
-	}
+				SNew(SComboButton)
+				.ComboButtonStyle(FAppStyle::Get(), "SimpleComboButton")
+				.OnGetMenuContent(this, &SObjectTableTreeView::MakeFiltersMenu)
+				.ToolTipText(LOCTEXT("FiltersMenuToolTip", "Filters"))
+				.ContentPadding(FMargin(0.0f, 1.0f))
+				.ButtonContent()
+				[
+					SNew(SImage)
+					.Image(FInsightsStyle::GetBrush("Icons.Filter.ToolBar"))
+				]
+			]
+
+			+ SHorizontalBox::Slot()
+			.FillWidth(1.0f)
+			.VAlign(VAlign_Center)
+			[
+				ConstructSearchBox()
+			]
+
+			+ SHorizontalBox::Slot()
+			.AutoWidth()
+			.Padding(2.0f, 0.0f, 0.0f, 0.0f)
+			.VAlign(VAlign_Center)
+			[
+				SNew(SBox)
+				.Visibility_Lambda([this]() { return bIsAdvancedFilterConfiguratorVisible ? EVisibility::Visible : EVisibility::Collapsed; })
+				[
+					ConstructFilterConfiguratorButton()
+				]
+			]
+
+#if 0 // TODO
+			+ SHorizontalBox::Slot()
+			.AutoWidth()
+			.Padding(2.0f, 0.0f, 0.0f, 0.0f)
+			.VAlign(VAlign_Center)
+			[
+				SNew(SButton)
+				.ContentPadding(FMargin(0.0f, 0.0f, 0.0f, 0.0f))
+				.OnClicked_Lambda([this]() -> FReply
+					{
+						ToggleObjectDetailsViewVisibility();
+						return FReply::Handled();
+					})
+				[
+					SNew(SImage)
+					.Image_Lambda([this]() -> const FSlateBrush*
+						{
+							if (IsObjectDetailsViewVisible())
+							{
+								return FInsightsStyle::Get().GetBrush("Icons.FindPrevious.ToolBar");
+								
+							}
+							else
+							{
+								return FInsightsStyle::Get().GetBrush("Icons.FindNext.ToolBar");
+							}
+						})
+				]
+			]
+#endif
+		];
+
+	//////////////////////////////////////////////////
+	// Grouping
+
+	InHostBox->AddSlot()
+		.AutoHeight()
+		.Padding(4.0f, 2.0f, 4.0f, 4.0f)
+		[
+			SNew(SBox)
+			.Visibility_Lambda([this]() { return bIsHierarchyBreadcrumbTrailVisible ? EVisibility::Visible : EVisibility::Collapsed; })
+			[
+				ConstructHierarchyBreadcrumbTrail()
+			]
+		];
 }
+END_SLATE_FUNCTION_BUILD_OPTIMIZATION
 
 ////////////////////////////////////////////////////////////////////////////////////////////////////
 
 BEGIN_SLATE_FUNCTION_BUILD_OPTIMIZATION
-TSharedPtr<SWidget> SObjectTableTreeView::ConstructFilterToolbar()
+TSharedPtr<SWidget> SObjectTableTreeView::ConstructFooter()
 {
-	TSharedRef<SHorizontalBox> FilterLine = SNew(SHorizontalBox)
+	return SNew(SHorizontalBox)
 
 	+ SHorizontalBox::Slot()
 	.FillWidth(1.0f)
-	.VAlign(VAlign_Center)
+
+	+ SHorizontalBox::Slot()
+	.AutoWidth()
+	.Padding(0.0f, 2.0f, 0.0f, 2.0f)
 	[
-		ConstructSearchBox()
+		SNew(STextBlock)
+		.Text(this, &SObjectTableTreeView::GetNumSelectedObjectsText)
 	]
 
 	+ SHorizontalBox::Slot()
 	.AutoWidth()
-	.Padding(2.0f, 0.0f, 0.0f, 0.0f)
-	.VAlign(VAlign_Center)
+	.Padding(4.0f, 2.0f, 0.0f, 2.0f)
 	[
-		ConstructFilterConfiguratorButton()
-	];
-
-	ConstructViewPreset(FilterLine, 140.0f);
-
-	return FilterLine;
-}
-END_SLATE_FUNCTION_BUILD_OPTIMIZATION
-
-////////////////////////////////////////////////////////////////////////////////////////////////////
-
-BEGIN_SLATE_FUNCTION_BUILD_OPTIMIZATION
-TSharedPtr<SWidget> SObjectTableTreeView::ConstructToolbar()
-{
-	return nullptr;
-}
-END_SLATE_FUNCTION_BUILD_OPTIMIZATION
+		SNew(STextBlock)
+		.Text(this, &SObjectTableTreeView::GetSelectedObjectsText)
+		.ColorAndOpacity(FSlateColor(EStyleColor::White25))
+	]
 
-////////////////////////////////////////////////////////////////////////////////////////////////////
+	+ SHorizontalBox::Slot()
+	.AutoWidth()
+	.Padding(0.0f, 2.0f, 0.0f, 2.0f)
+	[
+		SNew(STextBlock)
+		.Text(FText::FromString(TEXT(" (StructSize: ")))
+		.ColorAndOpacity(FSlateColor(EStyleColor::White25))
+		.Visibility_Lambda([this]() { return !NumSelectedObjectsText.IsEmpty() ? EVisibility::Visible : EVisibility::Collapsed; })
+	]
 
-BEGIN_SLATE_FUNCTION_BUILD_OPTIMIZATION
-TSharedPtr<SWidget> SObjectTableTreeView::ConstructFooter()
-{
-	return
-		SNew(SHorizontalBox)
+	+ SHorizontalBox::Slot()
+	.AutoWidth()
+	.Padding(0.0f, 2.0f, 0.0f, 2.0f)
+	[
+		SNew(STextBlock)
+		.Text(this, &SObjectTableTreeView::GetSelectionStructSizeText)
+	]
 
-		+ SHorizontalBox::Slot()
-		.FillWidth(1.0f)
+	+ SHorizontalBox::Slot()
+	.AutoWidth()
+	.Padding(0.0f, 2.0f, 0.0f, 2.0f)
+	[
+		SNew(STextBlock)
+		.Text(FText::FromString(TEXT(", SystemMemSize: ")))
+		.ColorAndOpacity(FSlateColor(EStyleColor::White25))
+		.Visibility_Lambda([this]() { return !NumSelectedObjectsText.IsEmpty() ? EVisibility::Visible : EVisibility::Collapsed; })
+	]
 
-		+ SHorizontalBox::Slot()
-		.AutoWidth()
-		.Padding(0.0f, 2.0f, 0.0f, 2.0f)
-		[
-			SNew(STextBlock)
-			.Text(this, &SObjectTableTreeView::GetNumSelectedObjectsText)
-		]
+	+ SHorizontalBox::Slot()
+	.AutoWidth()
+	.Padding(0.0f, 2.0f, 0.0f, 2.0f)
+	[
+		SNew(STextBlock)
+		.Text(this, &SObjectTableTreeView::GetSelectionSystemMemSizeText)
+	]
 
-		+ SHorizontalBox::Slot()
-		.AutoWidth()
-		.Padding(4.0f, 2.0f, 0.0f, 2.0f)
-		[
-			SNew(STextBlock)
-			.Text(this, &SObjectTableTreeView::GetSelectedObjectsText)
-			.ColorAndOpacity(FSlateColor(EStyleColor::White25))
-		]
-
-		+ SHorizontalBox::Slot()
-		.AutoWidth()
-		.Padding(0.0f, 2.0f, 0.0f, 2.0f)
-		[
-			SNew(STextBlock)
-			.Text(FText::FromString(TEXT(" (StructSize: ")))
-			.ColorAndOpacity(FSlateColor(EStyleColor::White25))
-			.Visibility_Lambda([this]() { return !NumSelectedObjectsText.IsEmpty() ? EVisibility::Visible : EVisibility::Collapsed; })
-		]
-
-		+ SHorizontalBox::Slot()
-		.AutoWidth()
-		.Padding(0.0f, 2.0f, 0.0f, 2.0f)
-		[
-			SNew(STextBlock)
-			.Text(this, &SObjectTableTreeView::GetSelectionStructSizeText)
-		]
+	+ SHorizontalBox::Slot()
+	.AutoWidth()
+	.Padding(0.0f, 2.0f, 0.0f, 2.0f)
+	[
+		SNew(STextBlock)
+		.Text(FText::FromString(TEXT(", VideoMemSize: ")))
+		.ColorAndOpacity(FSlateColor(EStyleColor::White25))
+		.Visibility_Lambda([this]() { return !NumSelectedObjectsText.IsEmpty() ? EVisibility::Visible : EVisibility::Collapsed; })
+	]
 
-		+ SHorizontalBox::Slot()
-		.AutoWidth()
-		.Padding(0.0f, 2.0f, 0.0f, 2.0f)
-		[
-			SNew(STextBlock)
-			.Text(FText::FromString(TEXT(", SystemMemSize: ")))
-			.ColorAndOpacity(FSlateColor(EStyleColor::White25))
-			.Visibility_Lambda([this]() { return !NumSelectedObjectsText.IsEmpty() ? EVisibility::Visible : EVisibility::Collapsed; })
-		]
-
-		+ SHorizontalBox::Slot()
-		.AutoWidth()
-		.Padding(0.0f, 2.0f, 0.0f, 2.0f)
-		[
-			SNew(STextBlock)
-			.Text(this, &SObjectTableTreeView::GetSelectionSystemMemSizeText)
-		]
+	+ SHorizontalBox::Slot()
+	.AutoWidth()
+	.Padding(0.0f, 2.0f, 0.0f, 2.0f)
+	[
+		SNew(STextBlock)
+		.Text(this, &SObjectTableTreeView::GetSelectionVideoMemSizeText)
+	]
 
-		+ SHorizontalBox::Slot()
-		.AutoWidth()
-		.Padding(0.0f, 2.0f, 0.0f, 2.0f)
-		[
-			SNew(STextBlock)
-			.Text(FText::FromString(TEXT(", VideoMemSize: ")))
-			.ColorAndOpacity(FSlateColor(EStyleColor::White25))
-			.Visibility_Lambda([this]() { return !NumSelectedObjectsText.IsEmpty() ? EVisibility::Visible : EVisibility::Collapsed; })
-		]
-
-		+ SHorizontalBox::Slot()
-		.AutoWidth()
-		.Padding(0.0f, 2.0f, 0.0f, 2.0f)
-		[
-			SNew(STextBlock)
-			.Text(this, &SObjectTableTreeView::GetSelectionVideoMemSizeText)
-		]
+	+ SHorizontalBox::Slot()
+	.AutoWidth()
+	.Padding(0.0f, 2.0f, 0.0f, 2.0f)
+	[
+		SNew(STextBlock)
+		.Text(FText::FromString(TEXT(")")))
+		.ColorAndOpacity(FSlateColor(EStyleColor::White25))
+		.Visibility_Lambda([this]() { return !NumSelectedObjectsText.IsEmpty() ? EVisibility::Visible : EVisibility::Collapsed; })
+	]
 
-		+ SHorizontalBox::Slot()
-		.AutoWidth()
-		.Padding(0.0f, 2.0f, 0.0f, 2.0f)
-		[
-			SNew(STextBlock)
-			.Text(FText::FromString(TEXT(")")))
-			.ColorAndOpacity(FSlateColor(EStyleColor::White25))
-			.Visibility_Lambda([this]() { return !NumSelectedObjectsText.IsEmpty() ? EVisibility::Visible : EVisibility::Collapsed; })
-		]
-
-		+ SHorizontalBox::Slot()
-		.FillWidth(1.0f);
+	+ SHorizontalBox::Slot()
+	.FillWidth(1.0f);
 }
 END_SLATE_FUNCTION_BUILD_OPTIMIZATION
 
 ////////////////////////////////////////////////////////////////////////////////////////////////////
 
-void SObjectTableTreeView::FindMatchedAsset(const TArray<FTableTreeNodePtr>& SelectedNodes, FAssetData& OutAsset) const
+bool SObjectTableTreeView::IsObjectDetailsViewVisible() const
 {
-	FAssetRegistryModule& AssetRegistryModule = FModuleManager::Get().LoadModuleChecked<FAssetRegistryModule>(TEXT("AssetRegistry"));
-
-	for (const FTableTreeNodePtr& SelectedNode : SelectedNodes)
-	{
-		if (SelectedNode->Is<FObjectNode>())
-		{
-			const FObjectNode& ObjectNode = SelectedNode->As<FObjectNode>();
+	return false; //TODO
+}
 
-			FString Path;
-			auto Package = ObjectNode.GetPackage();
-			if (Package)
-			{
-				Path = Package->GetObjectName();
-			}
-			else
-			{
-				Path = ObjectNode.GetObjectName();
-			}
+////////////////////////////////////////////////////////////////////////////////////////////////////
 
-			auto Outer = ObjectNode.GetOuter();
-			bool Result = false;
-			do
-			{
-				OutAsset = AssetRegistryModule.Get().GetAssetByObjectPath(Path, false, false);
-				if (OutAsset.IsValid())
-				{
-					return;
-				}
-				if (Outer.IsValid())
-				{
-					auto PackageOuter = Outer->GetPackage();
-					if (PackageOuter)
-					{
-						Path = PackageOuter->GetObjectPath();
-					}
-					else
-					{
-						Path = Outer->GetObjectName();
-					}
-					Outer = Outer->GetOuter();
-				}
-				else
-				{
-					break;
-				}
-			} while (!OutAsset.IsValid());
-		}
-	}
+void SObjectTableTreeView::ToggleObjectDetailsViewVisibility()
+{
+	//TODO
 }
 
 ////////////////////////////////////////////////////////////////////////////////////////////////////
@@ -311,6 +304,8 @@ void SObjectTableTreeView::UpdateSelectionStatsText()
 	TArray<FTableTreeNodePtr> SelectedNodes;
 	const int32 NumSelectedNodes = TreeView->GetSelectedItems(SelectedNodes);
 
+	TSharedPtr<const FObjectNode> FirstSelectedNode;
+
 	if (NumSelectedNodes > 0)
 	{
 		int64 TotalCount = 0;
@@ -321,6 +316,11 @@ void SObjectTableTreeView::UpdateSelectionStatsText()
 		{
 			if (SelectedNode->Is<FObjectNode>())
 			{
+				if (!FirstSelectedNode.IsValid())
+				{
+					FirstSelectedNode = StaticCastSharedPtr<const FObjectNode>(SelectedNode);
+				}
+
 				const FObjectNode& ObjectNode = SelectedNode->As<FObjectNode>();
 				++TotalCount;
 				TotalStructSize += ObjectNode.GetStructureSize();
@@ -346,6 +346,11 @@ void SObjectTableTreeView::UpdateSelectionStatsText()
 		SelectionSystemMemSizeText = FText::GetEmpty();
 		SelectionVideoMemSizeText = FText::GetEmpty();
 	}
+
+	if (TSharedPtr<SObjectDetailsView> ObjectDetailsView = WeakObjectDetailsView.Pin())
+	{
+		ObjectDetailsView->SetObject(FirstSelectedNode);
+	}
 }
 
 ////////////////////////////////////////////////////////////////////////////////////////////////////
@@ -353,18 +358,6 @@ void SObjectTableTreeView::UpdateSelectionStatsText()
 void SObjectTableTreeView::TreeView_OnSelectionChanged(FTableTreeNodePtr SelectedItem, ESelectInfo::Type SelectInfo)
 {
 	UpdateSelectionStatsText();
-
-	TArray<FTableTreeNodePtr> SelectedNodes;
-	TreeView->GetSelectedItems(SelectedNodes);
-
-	FAssetData Asset;
-	FindMatchedAsset(SelectedNodes, Asset);
-
-	TSharedPtr<SObjectDetailsView> DetailsViewShared = DetailsView.Pin();
-	if (DetailsViewShared)
-	{
-		DetailsViewShared->SetMatchedAssets(Asset);
-	}
 }
 
 ////////////////////////////////////////////////////////////////////////////////////////////////////
@@ -383,6 +376,20 @@ void SObjectTableTreeView::TreeView_OnMouseButtonDoubleClick(FTableTreeNodePtr N
 	}
 }
 
+////////////////////////////////////////////////////////////////////////////////////////////////////
+
+TSharedRef<SWidget> SObjectTableTreeView::MakeFiltersMenu()
+{
+	FMenuBuilder MenuBuilder(/*bInShouldCloseWindowAfterMenuSelection=*/true, nullptr);
+	MenuBuilder.SetSearchable(false);
+
+	MenuBuilder.BeginSection("Filters"/*, LOCTEXT("ContextMenu_Section_Trackers", "Trackers")*/);
+
+	MenuBuilder.EndSection();
+
+	return MenuBuilder.MakeWidget();
+}
+
 ////////////////////////////////////////////////////////////////////////////////////////////////////
 // Grouping
 ////////////////////////////////////////////////////////////////////////////////////////////////////
@@ -557,18 +564,32 @@ void SObjectTableTreeView::RebuildTree(bool bResync)
 
 					TableRowNodes.Reserve(ValidObjectCount);
 
+					// All object nodes are created...
 					for (TSharedPtr<FObjectNode>& NodePtr : ObjectNodes)
 					{
 						if (NodePtr.IsValid())
 						{
-							// Resolve pointers to class and outer objects.
+							// ...Now we can resolve pointers to class and outer objects.
 							NodePtr->SetClass(GetObjectNode(NodePtr->GetClassId()));
 							NodePtr->SetOuter(GetObjectNode(NodePtr->GetOuterId()));
+
+							// Also add the valid pointers to the tree view's list of nodes.
 							TableRowNodes.Add(NodePtr);
 						}
 					}
-				}
 
+					// The outer objects are set...
+					for (TSharedPtr<FObjectNode>& NodePtr : ObjectNodes)
+					{
+						if (NodePtr.IsValid())
+						{
+							// ...Now we can init the object paths.
+							NodePtr->InitObjectPath();
+						}
+					}
+
+					// TODO: asset matching
+				}
 			}
 		}
 	}
diff --git a/Engine/Source/Developer/TraceInsights/Private/Insights/ObjectProfiler/Widgets/SObjectTableTreeView.h b/Engine/Source/Developer/TraceInsights/Private/Insights/ObjectProfiler/Widgets/SObjectTableTreeView.h
--- a/Engine/Source/Developer/TraceInsights/Private/Insights/ObjectProfiler/Widgets/SObjectTableTreeView.h
+++ b/Engine/Source/Developer/TraceInsights/Private/Insights/ObjectProfiler/Widgets/SObjectTableTreeView.h
@@ -4,12 +4,14 @@
 
 #include "CoreTypes.h"
 
-#include "Templates/SharedPointer.h"
-
 // TraceInsights
-#include "Insights/ObjectProfiler/ViewModels/AssetInfoNode.h"
 #include "Insights/Table/Widgets/SSessionTableTreeView.h"
 
+namespace UE::Insights
+{
+	class SSegmentedBarGraph;
+}
+
 namespace UE::Insights::ObjectProfiler
 {
 
@@ -17,9 +19,6 @@ class FObjectNode;
 class FObjectTable;
 class SObjectDetailsView;
 
-DECLARE_DELEGATE_OneParam(FOnMatchedAssetsChanged, const TArray<TSharedPtr<FAssetInfoNode>>&);
-DECLARE_DELEGATE_RetVal_OneParam(bool, FIsAssetAvailable, const FAssetData& Asset);
-
 ////////////////////////////////////////////////////////////////////////////////////////////////////
 /**
  * A custom widget used to display the list of objects from a UObject Snapshot.
@@ -40,16 +39,23 @@ public:
 	 * Construct this widget
 	 * @param InArgs - The declaration data for this widget
 	 */
-	void Construct(const FArguments& InArgs, TWeakPtr<SObjectDetailsView> DetailsView);
+	void Construct(const FArguments& InArgs);
 
 	virtual void ConstructHeaderArea(TSharedRef<SVerticalBox> InHostBox) override;
-	virtual void ConstructFooterArea(TSharedRef<SVerticalBox> InHostBox) override;
-	virtual TSharedPtr<SWidget> ConstructFilterToolbar() override;
-	virtual TSharedPtr<SWidget> ConstructToolbar() override;
 	virtual TSharedPtr<SWidget> ConstructFooter() override;
 
 	virtual void Reset();
 
+	void SetSegmentedBarGraph(TWeakPtr<SSegmentedBarGraph> InSegmentedBarGraph)
+	{
+		WeakSegmentedBarGraph = InSegmentedBarGraph;
+	}
+
+	void SetObjectDetailsView(TWeakPtr<SObjectDetailsView> InObjectDetailsView)
+	{
+		WeakObjectDetailsView = InObjectDetailsView;
+	}
+
 	/**
 	 * Ticks this widget.  Override in derived classes, but always call the parent implementation.
 	 *
@@ -65,9 +71,6 @@ public:
 	 */
 	virtual void RebuildTree(bool bResync);
 
-	FOnMatchedAssetsChanged& GetOnMatchedAssetsChangedEvent() { return OnMatchedAssetsChangedEvent; }
-	FIsAssetAvailable& GetIsAssetAvailableEvent() { return IsAssetAvailableEvent; }
-
 protected:
 	virtual void InternalCreateGroupings() override;
 
@@ -82,8 +85,12 @@ protected:
 	virtual void TreeView_OnSelectionChanged(FTableTreeNodePtr SelectedItem, ESelectInfo::Type SelectInfo) override;
 
 private:
+	TSharedRef<SWidget> MakeFiltersMenu();
+
+	bool IsObjectDetailsViewVisible() const;
+	void ToggleObjectDetailsViewVisibility();
+
 	void UpdateSelectionStatsText();
-	void FindMatchedAsset(const TArray<FTableTreeNodePtr>& SelectedNodes, FAssetData& OutAsset) const;
 	FText GetNumSelectedObjectsText() const { return NumSelectedObjectsText; }
 	FText GetSelectedObjectsText() const { return  SelectedObjectsText; }
 	FText GetSelectionStructSizeText() const { return SelectionStructSizeText; }
@@ -112,18 +119,19 @@ private:
 	bool bHideObjectsWithLowEstimatedMemory = false; // hide objects with Exclusive Estimated Memory (System + Video) < 1 MiB
 	bool bHideObjectsWithLowImpact = false; // hide objects with Impact % < 1%
 
+	bool bIsAdvancedFilterConfiguratorVisible = true;
+	bool bIsHierarchyBreadcrumbTrailVisible = true;
+
 	//////////////////////////////////////////////////
 
+	TWeakPtr<SSegmentedBarGraph> WeakSegmentedBarGraph;
+	TWeakPtr<SObjectDetailsView> WeakObjectDetailsView;
+
 	FText NumSelectedObjectsText;
 	FText SelectedObjectsText;
 	FText SelectionStructSizeText;
 	FText SelectionSystemMemSizeText;
 	FText SelectionVideoMemSizeText;
-
-	FOnMatchedAssetsChanged OnMatchedAssetsChangedEvent;
-	FIsAssetAvailable IsAssetAvailableEvent;
-
-	TWeakPtr<SObjectDetailsView> DetailsView;
 };
 
 ////////////////////////////////////////////////////////////////////////////////////////////////////
diff --git a/Engine/Source/Developer/TraceInsights/Private/Insights/TraceInsightsModule.cpp b/Engine/Source/Developer/TraceInsights/Private/Insights/TraceInsightsModule.cpp
--- a/Engine/Source/Developer/TraceInsights/Private/Insights/TraceInsightsModule.cpp
+++ b/Engine/Source/Developer/TraceInsights/Private/Insights/TraceInsightsModule.cpp
@@ -37,6 +37,7 @@
 #include "Insights/NetworkingProfiler/NetworkingProfilerManager.h"
 #include "Insights/ObjectProfiler/ObjectProfilerManager.h"
 #include "Insights/ObjectProfiler/Widgets/SMemoryView.h"
+#include "Insights/ObjectProfiler/Widgets/SObjectProfilerWindow.h"
 #include "Insights/TaskGraphProfiler/TaskGraphProfilerManager.h"
 #include "Insights/Tests/InsightsTestRunner.h"
 #include "Insights/TimingProfiler/TimingProfilerManager.h"
@@ -628,6 +629,7 @@ void FTraceInsightsModule::HandleCodeAccessorOpenFileFailed(const FString& Filen
 TSharedRef<UE::Insights::ObjectProfiler::IObjectProfilerMemoryView> FTraceInsightsModule::CreateMemoryView() const
 {
 	return SNew(UE::Insights::ObjectProfiler::SMemoryView);
+	//return SNew(UE::Insights::ObjectProfiler::SObjectProfilerWindow);
 }
 
 ////////////////////////////////////////////////////////////////////////////////////////////////////
diff --git a/Engine/Source/Developer/TraceInsights/Public/Insights/Widgets/IObjectProfilerMemoryView.h b/Engine/Source/Developer/TraceInsights/Public/Insights/Widgets/IObjectProfilerMemoryView.h
--- a/Engine/Source/Developer/TraceInsights/Public/Insights/Widgets/IObjectProfilerMemoryView.h
+++ b/Engine/Source/Developer/TraceInsights/Public/Insights/Widgets/IObjectProfilerMemoryView.h
@@ -2,23 +2,24 @@
 
 #pragma once
 
-#include "AssetRegistry/AssetData.h"
-#include "Containers/StringFwd.h"
-#include "Widgets/DeclarativeSyntaxSupport.h"
-#include "Widgets/SCompoundWidget.h"
+#include "CoreTypes.h"
+
+#include "UObject/SoftObjectPath.h"
 
 #define UE_API TRACEINSIGHTS_API
 
-////////////////////////////////////////////////////////////////////////////////////////////////////
+class SWidget;
+struct FAssetData;
 
 namespace UE::Insights::ObjectProfiler
 {
 
+class IAssetInfoProvider;
+
 /** Interface for the Object Profiler Memory View. */
 class IObjectProfilerMemoryView
 {
 public:
-
 	/** Default constructor. */
 	IObjectProfilerMemoryView() {}
 
@@ -26,9 +27,12 @@ public:
 	 virtual ~IObjectProfilerMemoryView() {}
 
 	virtual TSharedRef<SWidget> GetWidget() = 0;
-	
+
+	// deprecated; to be removed
 	DECLARE_DELEGATE_RetVal_TwoParams(bool, FGetActors, const FAssetData& Asset, TArray<FSoftObjectPath>& OutActors);
-	virtual void SetGetActorsDelegate(FGetActors InDelegate) = 0;
+	virtual void SetGetActorsDelegate(FGetActors InDelegate) {}
+
+	virtual void SetAssetInfoProvider(TSharedPtr<IAssetInfoProvider> InAssetInfoProvider) = 0;
 };
 
 } // namespace UE::Insights::ObjectProfiler
