[Insights] Fix for assert when exporting TimerCallees
* Added ReadScope to ExportTimerCalleesAsText

#Synced-CL 52134467

diff --git a/Engine/Source/Developer/TraceInsights/Private/Insights/TimingProfiler/ViewModels/TimingExporter.cpp b/Engine/Source/Developer/TraceInsights/Private/Insights/TimingProfiler/ViewModels/TimingExporter.cpp
--- a/Engine/Source/Developer/TraceInsights/Private/Insights/TimingProfiler/ViewModels/TimingExporter.cpp
+++ b/Engine/Source/Developer/TraceInsights/Private/Insights/TimingProfiler/ViewModels/TimingExporter.cpp
@@ -1141,6 +1141,7 @@ int32 FTimingExporter::ExportTimerCalleesAsText(const FString& Filename, const F
 
 	// Write rows
 	{
+		TraceServices::FAnalysisSessionReadScope SessionReadScope(Session);
 		for (uint32 TimerId : Params.TimerIds)
 		{
 			TArray<const TraceServices::FTimingProfilerButterflyNode*> NodesToVisit{ &Butterfly->GenerateCalleesTree(TimerId) };
