[Insights] Object Profiler fixes.
* SMemoryView: Fixed propagation of AssetInfoProvider to ObjectDetailsView.
* Object Details: Added validation for GetClass() usages.
* Updated the debug test for SSegmentedBarGraph.

#rnx
#Synced-CL 52097700

diff --git a/Engine/Source/Developer/TraceInsights/Private/Insights/ObjectProfiler/Widgets/SMemoryView.cpp b/Engine/Source/Developer/TraceInsights/Private/Insights/ObjectProfiler/Widgets/SMemoryView.cpp
--- a/Engine/Source/Developer/TraceInsights/Private/Insights/ObjectProfiler/Widgets/SMemoryView.cpp
+++ b/Engine/Source/Developer/TraceInsights/Private/Insights/ObjectProfiler/Widgets/SMemoryView.cpp
@@ -49,11 +49,7 @@ void SMemoryView::Construct(const FArguments& InArgs)
 			SNew(SBox)
 			.HeightOverride(22.0f)
 			[
-				SNew(SBorder)
-				.BorderBackgroundColor(FLinearColor(0.7f, 0.3f, 0.3f, 1.0f))
-				[
-					SAssignNew(SegmentedBarGraph, SSegmentedBarGraph)
-				]
+				SAssignNew(SegmentedBarGraph, SSegmentedBarGraph)
 			]
 		]
 
@@ -79,7 +75,11 @@ void SMemoryView::Construct(const FArguments& InArgs)
 	ObjectTableTreeView->SetObjectDetailsView(ObjectDetailsView.ToWeakPtr());
 	ObjectTableTreeView->SetSegmentedBarGraph(SegmentedBarGraph.ToWeakPtr());
 	ObjectDetailsView->SetParentMemoryView(SharedThis(this));
-	ObjectDetailsView->SetAssetInfoProvider(AssetInfoProvider);
+
+	if (AssetInfoProvider)
+	{
+		ObjectDetailsView->SetAssetInfoProvider(AssetInfoProvider);
+	}
 }
 
 ////////////////////////////////////////////////////////////////////////////////////////////////////
@@ -101,6 +101,18 @@ void SMemoryView::OnClose()
 
 ////////////////////////////////////////////////////////////////////////////////////////////////////
 
+void SMemoryView::SetAssetInfoProvider(TSharedPtr<IAssetInfoProvider> InAssetInfoProvider)
+{
+	AssetInfoProvider = InAssetInfoProvider;
+
+	if (AssetInfoProvider && ObjectDetailsView)
+	{
+		ObjectDetailsView->SetAssetInfoProvider(AssetInfoProvider);
+	}
+}
+
+////////////////////////////////////////////////////////////////////////////////////////////////////
+
 } // namespace UE::Insights::ObjectProfiler
 
 #undef LOCTEXT_NAMESPACE
diff --git a/Engine/Source/Developer/TraceInsights/Private/Insights/ObjectProfiler/Widgets/SMemoryView.h b/Engine/Source/Developer/TraceInsights/Private/Insights/ObjectProfiler/Widgets/SMemoryView.h
--- a/Engine/Source/Developer/TraceInsights/Private/Insights/ObjectProfiler/Widgets/SMemoryView.h
+++ b/Engine/Source/Developer/TraceInsights/Private/Insights/ObjectProfiler/Widgets/SMemoryView.h
@@ -51,10 +51,7 @@ public:
 
 	virtual TSharedRef<SWidget> GetWidget() override { return SharedThis(this); }
 
-	virtual void SetAssetInfoProvider(TSharedPtr<IAssetInfoProvider> InAssetInfoProvider) override
-	{
-		AssetInfoProvider = InAssetInfoProvider;
-	}
+	virtual void SetAssetInfoProvider(TSharedPtr<IAssetInfoProvider> InAssetInfoProvider) override;
 
 private:
 	TSharedPtr<IAssetInfoProvider> AssetInfoProvider;
diff --git a/Engine/Source/Developer/TraceInsights/Private/Insights/ObjectProfiler/Widgets/SObjectDetailsView.cpp b/Engine/Source/Developer/TraceInsights/Private/Insights/ObjectProfiler/Widgets/SObjectDetailsView.cpp
--- a/Engine/Source/Developer/TraceInsights/Private/Insights/ObjectProfiler/Widgets/SObjectDetailsView.cpp
+++ b/Engine/Source/Developer/TraceInsights/Private/Insights/ObjectProfiler/Widgets/SObjectDetailsView.cpp
@@ -319,7 +319,8 @@ void SObjectDetailsView::SetObject(TSharedPtr<const FObjectNode> InObjectNode)
 		if (AssetInfoProvider)
 		{
 			AssetInfoProvider->GetAssetData(*InObjectNode.Get(), SelectedAssetData);
-			SelectedAssetClassName = FName(InObjectNode->GetClass()->GetObjectName());
+			TSharedPtr<FObjectNode> ClassNode = InObjectNode->GetClass();
+			SelectedAssetClassName = ClassNode ? FName(ClassNode->GetObjectName()) : NAME_None;
 
 			if (SelectedAssetData.IsValid())
 			{
@@ -417,12 +418,15 @@ FText SObjectDetailsView::GetSelectedObjectThumbnailToolTip() const
 		Str.Append(SelectedObjectNode->GetObjectName());
 		Str.Append("\nId: ");
 		Str.Appendf("%u", SelectedObjectNode->GetObjectId());
-		Str.Append("\nClass: ");
-		Str.Append(SelectedObjectNode->GetClass()->GetObjectName());
-		if (SelectedObjectNode->GetOuter().IsValid())
+		if (TSharedPtr<FObjectNode> ClassNode = SelectedObjectNode->GetClass())
+		{
+			Str.Append("\nClass: ");
+			Str.Append(ClassNode->GetObjectName());
+		}
+		if (TSharedPtr<FObjectNode> OuterNode = SelectedObjectNode->GetOuter())
 		{
 			Str.Append("\nOuter: ");
-			Str.Append(SelectedObjectNode->GetOuter()->GetObjectName());
+			Str.Append(OuterNode->GetObjectName());
 		}
 		if (SelectedObjectNode->GetSystemMemorySize() > 0)
 		{
diff --git a/Engine/Source/Developer/TraceInsights/Private/Insights/ObjectProfiler/Widgets/SObjectProfilerWindow.cpp b/Engine/Source/Developer/TraceInsights/Private/Insights/ObjectProfiler/Widgets/SObjectProfilerWindow.cpp
--- a/Engine/Source/Developer/TraceInsights/Private/Insights/ObjectProfiler/Widgets/SObjectProfilerWindow.cpp
+++ b/Engine/Source/Developer/TraceInsights/Private/Insights/ObjectProfiler/Widgets/SObjectProfilerWindow.cpp
@@ -251,29 +251,29 @@ TSharedRef<SWidget> SObjectProfilerWindow::CreateToolbar(TSharedPtr<FExtender> E
 
 	//return SNew(SObjectProfilerToolbar, Config).ToolbarExtender(Extender);
 
-#if 0 // debug code
+#if 1 // debug code
 	auto Segment1 = MakeShared<FCustomBarGraphSegment>();
-	Segment1->Size = 40.0;
+	Segment1->Size = 50.0;
 	Segment1->Text = LOCTEXT("Segment1", "Textures");
-	Segment1->ToolTipText = LOCTEXT("Segment1_TT", "All textures");
-	Segment1->Color = FLinearColor(0.5f, 0.0f, 0.0f, 1.0f);
-	Segment1->TextColor = FLinearColor(1.0f, 0.0f, 0.0f, 1.0f);
+	Segment1->ToolTipText = LOCTEXT("Segment1_TT", "Textures...");
+	Segment1->Color = FLinearColor(0.4f, 0.0f, 0.1f, 1.0f);
+	Segment1->TextColor = FLinearColor(1.0f, 1.0f, 1.0f, 1.0f);
 	Segments.Add(Segment1);
 
 	auto Segment2 = MakeShared<FCustomBarGraphSegment>();
-	Segment2->Size = 40.0;
+	Segment2->Size = 30.0;
 	Segment2->Text = LOCTEXT("Segment2", "Materials");
-	Segment2->ToolTipText = LOCTEXT("Segment2_TT", "All materials");
-	Segment2->Color = FLinearColor(0.0f, 0.3f, 0.0f, 1.0f);
-	Segment2->TextColor = FLinearColor(1.0f, 0.0f, 0.0f, 1.0f);
+	Segment2->ToolTipText = LOCTEXT("Segment2_TT", "Materials...");
+	Segment2->Color = FLinearColor(0.0f, 0.2f, 0.0f, 1.0f);
+	Segment2->TextColor = FLinearColor(1.0f, 1.0f, 1.0f, 1.0f);
 	Segments.Add(Segment2);
 
 	auto Segment3 = MakeShared<FCustomBarGraphSegment>();
-	Segment3->Size = 40.0;
+	Segment3->Size = 20.0;
 	Segment3->Text = LOCTEXT("Segment3_Verse", "Verse");
 	Segment3->ToolTipText = LOCTEXT("Segment3_TT", "Verse objects");
-	Segment3->Color = FLinearColor(0.0f, 0.0f, 0.5f, 1.0f);
-	Segment3->TextColor = FLinearColor(1.0f, 0.0f, 0.0f, 1.0f);
+	Segment3->Color = FLinearColor(0.1f, 0.0f, 0.4f, 1.0f);
+	Segment3->TextColor = FLinearColor(1.0f, 1.0f, 1.0f, 1.0f);
 	Segments.Add(Segment3);
 #endif
 
@@ -296,13 +296,8 @@ TSharedRef<SWidget> SObjectProfilerWindow::CreateToolbar(TSharedPtr<FExtender> E
 				SNew(SBox)
 				.HeightOverride(22.0f)
 				[
-					SNew(SBorder)
-					.BorderImage(FAppStyle::Get().GetBrush("WhiteBrush"))
-					.BorderBackgroundColor(FLinearColor(0.0f, 0.1f, 0.0f, 1.0f))
-					[
-						SAssignNew(SegmentedBarGraph, SSegmentedBarGraph)
-						//.SegmentsSource(&Segments)
-					]
+					SAssignNew(SegmentedBarGraph, SSegmentedBarGraph)
+					.SegmentsSource(&Segments)
 				]
 			]
 		];
diff --git a/Engine/Source/Developer/TraceInsights/Private/Insights/ObjectProfiler/Widgets/SObjectProfilerWindow.h b/Engine/Source/Developer/TraceInsights/Private/Insights/ObjectProfiler/Widgets/SObjectProfilerWindow.h
--- a/Engine/Source/Developer/TraceInsights/Private/Insights/ObjectProfiler/Widgets/SObjectProfilerWindow.h
+++ b/Engine/Source/Developer/TraceInsights/Private/Insights/ObjectProfiler/Widgets/SObjectProfilerWindow.h
@@ -119,7 +119,7 @@ private:
 
 	TSharedPtr<SSegmentedBarGraph> SegmentedBarGraph;
 	bool bIsSegmentedBarGraphVisible = true;
-	TArray<TSharedPtr<FCustomBarGraphSegment>> Segments;
+	TArray<TSharedPtr<IBarGraphSegment>> Segments;
 };
 
 ////////////////////////////////////////////////////////////////////////////////////////////////////
