[Insights] Object Profiler: Fixed major tab to be docked by default.

#rnx
#Synced-CL 52201748

diff --git a/Engine/Source/Developer/TraceInsights/Private/Insights/TraceInsightsModule.cpp b/Engine/Source/Developer/TraceInsights/Private/Insights/TraceInsightsModule.cpp
--- a/Engine/Source/Developer/TraceInsights/Private/Insights/TraceInsightsModule.cpp
+++ b/Engine/Source/Developer/TraceInsights/Private/Insights/TraceInsightsModule.cpp
@@ -310,7 +310,7 @@ void FTraceInsightsModule::CreateSessionViewer(bool bAllowDebugTools)
 	//////////////////////////////////////////////////
 	// Setup the window's content.
 
-	TSharedRef<FTabManager::FLayout> DefaultLayout = FTabManager::NewLayout("UnrealInsightsLayout_v1.1");
+	TSharedRef<FTabManager::FLayout> DefaultLayout = FTabManager::NewLayout("UnrealInsightsLayout_v1.2");
 
 	AddAreaForSessionViewer(DefaultLayout);
 
@@ -366,6 +366,7 @@ void FTraceInsightsModule::AddAreaForSessionViewer(TSharedRef<FTabManager::FLayo
 	Stack->AddTab(FInsightsManagerTabs::LoadingProfilerTabId, ETabState::ClosedTab);
 	Stack->AddTab(FInsightsManagerTabs::NetworkingProfilerTabId, ETabState::ClosedTab);
 	Stack->AddTab(FInsightsManagerTabs::MemoryProfilerTabId, ETabState::ClosedTab);
+	Stack->AddTab(UE::Insights::ObjectProfiler::SObjectProfilerWindow::ObjectProfilerTabId, ETabState::ClosedTab);
 	//Stack->SetForegroundTab(FTabId(FInsightsManagerTabs::TimingProfilerTabId));
 
 	// Create area for the main window.
@@ -382,6 +383,7 @@ void FTraceInsightsModule::AddAreaForSessionViewer(TSharedRef<FTabManager::FLayo
 	Stack->AddTab(FInsightsManagerTabs::NetworkingProfilerTabId, ETabState::ClosedTab);
 	Stack->AddTab(FInsightsManagerTabs::MemoryProfilerTabId, ETabState::ClosedTab);
 	Stack->AddTab(FInsightsManagerTabs::MessageLogTabId, ETabState::ClosedTab);
+	Stack->AddTab(UE::Insights::ObjectProfiler::SObjectProfilerWindow::ObjectProfilerTabId, ETabState::ClosedTab);
 	Stack->SetForegroundTab(FTabId(FInsightsManagerTabs::TimingProfilerTabId));
 
 	Layout->AddArea
