[Insights] Object Profiler: Fixed propagation of AssetInfoProvider to child views created by SObjectProfilerWindow.

#rnx
#Synced-CL 52248455

diff --git a/Engine/Source/Developer/TraceInsights/Private/Insights/ObjectProfiler/Widgets/SObjectProfilerWindow.cpp b/Engine/Source/Developer/TraceInsights/Private/Insights/ObjectProfiler/Widgets/SObjectProfilerWindow.cpp
--- a/Engine/Source/Developer/TraceInsights/Private/Insights/ObjectProfiler/Widgets/SObjectProfilerWindow.cpp
+++ b/Engine/Source/Developer/TraceInsights/Private/Insights/ObjectProfiler/Widgets/SObjectProfilerWindow.cpp
@@ -67,6 +67,26 @@ void SObjectProfilerWindow::Reset()
 
 ////////////////////////////////////////////////////////////////////////////////////////////////////
 
+void SObjectProfilerWindow::SetAssetInfoProvider(TSharedPtr<IAssetInfoProvider> InAssetInfoProvider)
+{
+	AssetInfoProvider = InAssetInfoProvider;
+
+	if (AssetInfoProvider)
+	{
+		if (ObjectTableTreeView)
+		{
+			ObjectTableTreeView->SetAssetInfoProvider(AssetInfoProvider);
+		}
+
+		if (ObjectDetailsView)
+		{
+			ObjectDetailsView->SetAssetInfoProvider(AssetInfoProvider);
+		}
+	}
+}
+
+////////////////////////////////////////////////////////////////////////////////////////////////////
+
 BEGIN_SLATE_FUNCTION_BUILD_OPTIMIZATION
 TSharedRef<SDockTab> SObjectProfilerWindow::SpawnTab_ObjectTableTreeView(const FSpawnTabArgs& Args)
 {
@@ -83,6 +103,11 @@ TSharedRef<SDockTab> SObjectProfilerWindow::SpawnTab_ObjectTableTreeView(const F
 
 	ObjectTableTreeView->SetLogListingName(FObjectProfilerManager::Get()->GetLogListingName());
 
+	if (AssetInfoProvider)
+	{
+		ObjectTableTreeView->SetAssetInfoProvider(AssetInfoProvider);
+	}
+
 	if (ObjectDetailsView)
 	{
 		ObjectTableTreeView->SetObjectDetailsView(ObjectDetailsView);
@@ -127,7 +152,10 @@ TSharedRef<SDockTab> SObjectProfilerWindow::SpawnTab_ObjectDetailsView(const FSp
 		];
 	check(ObjectDetailsView != nullptr);
 
-	ObjectDetailsView->SetAssetInfoProvider(AssetInfoProvider);
+	if (AssetInfoProvider)
+	{
+		ObjectDetailsView->SetAssetInfoProvider(AssetInfoProvider);
+	}
 
 	if (ObjectTableTreeView)
 	{
@@ -305,26 +333,6 @@ TSharedRef<SWidget> SObjectProfilerWindow::CreateToolbar(TSharedPtr<FExtender> E
 
 ////////////////////////////////////////////////////////////////////////////////////////////////////
 
-void SObjectProfilerWindow::SetAssetInfoProvider(TSharedPtr<IAssetInfoProvider> InAssetInfoProvider)
-{
-	AssetInfoProvider = InAssetInfoProvider;
-
-	if (AssetInfoProvider)
-	{
-		if (ObjectTableTreeView)
-		{
-			ObjectTableTreeView->SetAssetInfoProvider(AssetInfoProvider);
-		}
-
-		if (ObjectDetailsView)
-		{
-			ObjectDetailsView->SetAssetInfoProvider(AssetInfoProvider);
-		}
-	}
-}
-
-////////////////////////////////////////////////////////////////////////////////////////////////////
-
 } // namespace UE::Insights::ObjectProfiler
 
 #undef LOCTEXT_NAMESPACE
