[Insights] Object Profiler:
* Changed computation for Impact of Estimated Mem Size to be always relative to the root node (instead of direct parent).
* Added FObjectNode static accesors for EstimatedMemoryImpact and NumReferences (i.e. for any node with aggregated stats).
* Moved the "inclusive" Memory Size columns after all the "exclusive" columns.

#rnx
#Synced-CL 52248455

diff --git a/Engine/Source/Developer/TraceInsights/Private/Insights/ObjectProfiler/ViewModels/ObjectNode.cpp b/Engine/Source/Developer/TraceInsights/Private/Insights/ObjectProfiler/ViewModels/ObjectNode.cpp
--- a/Engine/Source/Developer/TraceInsights/Private/Insights/ObjectProfiler/ViewModels/ObjectNode.cpp
+++ b/Engine/Source/Developer/TraceInsights/Private/Insights/ObjectProfiler/ViewModels/ObjectNode.cpp
@@ -185,6 +185,39 @@ int64 FObjectNode::GetEstimatedMemorySize(const FBaseTreeNode& Node)
 	return 0;
 }
 
+///////////////////////////////////////////////////////////////////////////////////////////////////
+
+// static
+TOptional<double> FObjectNode::GetEstimatedMemoryImpact(const FBaseTreeNode& Node)
+{
+	const int64 EstimatedSize = FObjectNode::GetEstimatedMemorySize(Node);
+	if (EstimatedSize > 0)
+	{
+		TSharedPtr<FBaseTreeNode> ParentNode = Node.GetParent();
+		if (ParentNode.IsValid())
+		{
+			// GetRoot
+			while (true)
+			{
+				TSharedPtr<FBaseTreeNode> ParentParentNode = ParentNode->GetParent();
+				if (!ParentParentNode.IsValid())
+				{
+					break;
+				}
+				ParentNode = ParentParentNode;
+			}
+
+			const int64 ParentEstimatedMemSize = FObjectNode::GetEstimatedMemorySize(*ParentNode);
+			if (ParentEstimatedMemSize > 0)
+			{
+				const double Ratio = double(EstimatedSize) / double(ParentEstimatedMemSize);
+				return Ratio;
+			}
+		}
+	}
+	return TOptional<double>();
+}
+
 ////////////////////////////////////////////////////////////////////////////////////////////////////
 
 double FObjectNode::GetEstimatedMemoryImpact() const
@@ -195,6 +228,17 @@ double FObjectNode::GetEstimatedMemoryImpact() const
 		TSharedPtr<FBaseTreeNode> ParentNode = GetParent();
 		if (ParentNode.IsValid())
 		{
+			// GetRoot
+			while (true)
+			{
+				TSharedPtr<FBaseTreeNode> ParentParentNode = ParentNode->GetParent();
+				if (!ParentParentNode.IsValid())
+				{
+					break;
+				}
+				ParentNode = ParentParentNode;
+			}
+
 			const int64 ParentEstimatedMemSize = FObjectNode::GetEstimatedMemorySize(*ParentNode);
 			if (ParentEstimatedMemSize > 0)
 			{
@@ -207,6 +251,32 @@ double FObjectNode::GetEstimatedMemoryImpact() const
 
 ////////////////////////////////////////////////////////////////////////////////////////////////////
 
+// static
+uint32 FObjectNode::GetNumReferences(const FBaseTreeNode& Node)
+{
+	if (Node.Is<FObjectNode>())
+	{
+		return Node.As<FObjectNode>().GetNumReferences();
+	}
+
+	if (Node.IsGroup() && Node.Is<FTableTreeNode>())
+	{
+		const FTableTreeNode& TableNode = Node.As<FTableTreeNode>();
+		if (TableNode.HasAggregatedValue(FObjectTableColumns::ReferencesColumnId))
+		{
+			const FTableCellValue& NumRefs = TableNode.GetAggregatedValue(FObjectTableColumns::ReferencesColumnId);
+			if (NumRefs.DataType == ETableCellDataType::Int64)
+			{
+				return uint32(NumRefs.Int64);
+			}
+		}
+	}
+
+	return 0;
+}
+
+////////////////////////////////////////////////////////////////////////////////////////////////////
+
 FAssetData FObjectNode::MakeAssetData()
 {
 	FString LongPackageName;
diff --git a/Engine/Source/Developer/TraceInsights/Private/Insights/ObjectProfiler/ViewModels/ObjectNode.h b/Engine/Source/Developer/TraceInsights/Private/Insights/ObjectProfiler/ViewModels/ObjectNode.h
--- a/Engine/Source/Developer/TraceInsights/Private/Insights/ObjectProfiler/ViewModels/ObjectNode.h
+++ b/Engine/Source/Developer/TraceInsights/Private/Insights/ObjectProfiler/ViewModels/ObjectNode.h
@@ -86,15 +86,18 @@ public:
 
 	int64 GetSystemMemorySize() const { return SystemMemorySize; }
 	int64 GetVideoMemorySize() const { return VideoMemorySize; }
-	int64 GetEstimatedMemorySize() const { return SystemMemorySize + VideoMemorySize; }
 
 	static int64 GetEstimatedMemorySize(const FBaseTreeNode& Node);
+	int64 GetEstimatedMemorySize() const { return SystemMemorySize + VideoMemorySize; }
+
+	static TOptional<double> GetEstimatedMemoryImpact(const FBaseTreeNode& Node);
 	double GetEstimatedMemoryImpact() const;
 
 	int64 GetInclusiveSystemMemorySize() const { return InclSystemMemorySize; }
 	int64 GetInclusiveVideoMemorySize() const { return InclVideoMemorySize; }
 	int64 GetInclusiveEstimatedMemorySize() const { return InclSystemMemorySize + InclVideoMemorySize; }
 
+	static uint32 GetNumReferences(const FBaseTreeNode& Node);
 	virtual uint32 GetNumReferences() const { return MatchedActors.Num(); }
 
 	bool IsMissing() const { return (uint32(InternalFlags) & uint32(EInternalFlags::IsMissing)) != 0; }
diff --git a/Engine/Source/Developer/TraceInsights/Private/Insights/ObjectProfiler/ViewModels/ObjectTable.cpp b/Engine/Source/Developer/TraceInsights/Private/Insights/ObjectProfiler/ViewModels/ObjectTable.cpp
--- a/Engine/Source/Developer/TraceInsights/Private/Insights/ObjectProfiler/ViewModels/ObjectTable.cpp
+++ b/Engine/Source/Developer/TraceInsights/Private/Insights/ObjectProfiler/ViewModels/ObjectTable.cpp
@@ -42,11 +42,11 @@
 const FName FObjectTableColumns::StructureSizeColumnId("StructureSize");
 
 const FName FObjectTableColumns::SystemMemorySizeColumnId("SystemMemorySize");
-const FName FObjectTableColumns::InclSystemMemorySizeColumnId("InclSystemMemorySize");
 const FName FObjectTableColumns::VideoMemorySizeColumnId("VideoMemorySize");
-const FName FObjectTableColumns::InclVideoMemorySizeColumnId("InclVideoMemorySize");
 const FName FObjectTableColumns::EstimatedSizeColumnId("EstimatedSize");
 const FName FObjectTableColumns::ImpactColumnId("Impact");
+const FName FObjectTableColumns::InclSystemMemorySizeColumnId("InclSystemMemorySize");
+const FName FObjectTableColumns::InclVideoMemorySizeColumnId("InclVideoMemorySize");
 const FName FObjectTableColumns::InclEstimatedSizeColumnId("InclEstimatedSize");
 
 const FName FObjectTableColumns::ReferencesColumnId("References");
@@ -1062,56 +1062,6 @@
 		AddColumn(ColumnRef);
 	}
 	//////////////////////////////////////////////////
-	// Inclusive System Memory Size Column
-	{
-		TSharedRef<FTableColumn> ColumnRef = MakeShared<FTableColumn>(FObjectTableColumns::InclSystemMemorySizeColumnId);
-		FTableColumn& Column = *ColumnRef;
-
-		Column.SetIndex(ColumnIndex++);
-
-		Column.SetShortName(LOCTEXT("InclSystemMemorySizeColumnName", "Incl. Mem Size"));
-		Column.SetTitleName(LOCTEXT("InclSystemMemorySizeColumnTitle", "Inclusive System Memory Size"));
-		Column.SetDescription(LOCTEXT("InclSystemMemorySizeColumnDesc", "The total estimated system memory size (bytes) of the UObject instance (including references)"));
-
-		Column.SetFlags(ETableColumnFlags::CanBeHidden |
-						ETableColumnFlags::CanBeFiltered);
-
-		Column.SetHorizontalAlignment(HAlign_Right);
-		Column.SetInitialWidth(100.0f);
-
-		Column.SetDataType(ETableCellDataType::Int64);
-
-		class FValueGetter : public FAggregatedValueGetter
-		{
-		public:
-			virtual const TOptional<FTableCellValue> GetNodeValue(const FBaseTreeNode& Node) const override
-			{
-				if (Node.Is<FObjectNode>())
-				{
-					const FObjectNode& ObjectNode = Node.As<FObjectNode>();
-					return FTableCellValue(int64(ObjectNode.GetInclusiveSystemMemorySize()));
-				}
-				return TOptional<FTableCellValue>();
-			}
-		};
-		TSharedRef<ITableCellValueGetter> Getter = MakeShared<FValueGetter>();
-		Column.SetValueGetter(Getter);
-
-		TSharedRef<ITableCellValueFormatter> Formatter = MakeShared<FInt64ValueFormatterAsMemory>();
-		Column.SetValueFormatter(Formatter);
-
-		TSharedRef<ITableCellValueSorter> Sorter = MakeShared<FSorterByInt64Value>(ColumnRef);
-		Column.SetValueSorter(Sorter);
-		Column.SetInitialSortMode(EColumnSortMode::Descending);
-
-		//TSharedRef<IFilterValueConverter> Converter = MakeShared<FMemoryFilterValueConverter>();
-		//Column.SetValueConverter(Converter);
-
-		Column.SetAggregation(ETableColumnAggregation::None);
-
-		AddColumn(ColumnRef);
-	}
-	//////////////////////////////////////////////////
 	// Video Memory Size Column
 	{
 		TSharedRef<FTableColumn> ColumnRef = MakeShared<FTableColumn>(FObjectTableColumns::VideoMemorySizeColumnId);
@@ -1163,56 +1113,6 @@
 		AddColumn(ColumnRef);
 	}
 	//////////////////////////////////////////////////
-	// Inclusive Video Memory Size Column
-	{
-		TSharedRef<FTableColumn> ColumnRef = MakeShared<FTableColumn>(FObjectTableColumns::InclVideoMemorySizeColumnId);
-		FTableColumn& Column = *ColumnRef;
-
-		Column.SetIndex(ColumnIndex++);
-
-		Column.SetShortName(LOCTEXT("InclVideoMemorySizeColumnName", "Incl. Video Size"));
-		Column.SetTitleName(LOCTEXT("InclVideoMemorySizeColumnTitle", "Inclusive Video Memory Size"));
-		Column.SetDescription(LOCTEXT("InclVideoMemorySizeColumnDesc", "The total estimated video memory size (bytes) of the UObject instance (including references)"));
-
-		Column.SetFlags(ETableColumnFlags::CanBeHidden |
-						ETableColumnFlags::CanBeFiltered);
-
-		Column.SetHorizontalAlignment(HAlign_Right);
-		Column.SetInitialWidth(100.0f);
-
-		Column.SetDataType(ETableCellDataType::Int64);
-
-		class FValueGetter : public FAggregatedValueGetter
-		{
-		public:
-			virtual const TOptional<FTableCellValue> GetNodeValue(const FBaseTreeNode& Node) const override
-			{
-				if (Node.Is<FObjectNode>())
-				{
-					const FObjectNode& ObjectNode = Node.As<FObjectNode>();
-					return FTableCellValue(int64(ObjectNode.GetInclusiveVideoMemorySize()));
-				}
-				return TOptional<FTableCellValue>();
-			}
-		};
-		TSharedRef<ITableCellValueGetter> Getter = MakeShared<FValueGetter>();
-		Column.SetValueGetter(Getter);
-
-		TSharedRef<ITableCellValueFormatter> Formatter = MakeShared<FInt64ValueFormatterAsMemory>();
-		Column.SetValueFormatter(Formatter);
-
-		TSharedRef<ITableCellValueSorter> Sorter = MakeShared<FSorterByInt64Value>(ColumnRef);
-		Column.SetValueSorter(Sorter);
-		Column.SetInitialSortMode(EColumnSortMode::Descending);
-
-		//TSharedRef<IFilterValueConverter> Converter = MakeShared<FMemoryFilterValueConverter>();
-		//Column.SetValueConverter(Converter);
-
-		Column.SetAggregation(ETableColumnAggregation::None);
-
-		AddColumn(ColumnRef);
-	}
-	//////////////////////////////////////////////////
 	// Estimated Size Column
 	{
 		TSharedRef<FTableColumn> ColumnRef = MakeShared<FTableColumn>(FObjectTableColumns::EstimatedSizeColumnId);
@@ -1325,19 +1225,10 @@
 		public:
 			virtual const TOptional<FTableCellValue> GetValue(const FTableColumn& Column, const FBaseTreeNode& Node) const override
 			{
-				const int64 EstimatedSize = FObjectNode::GetEstimatedMemorySize(Node);
-				if (EstimatedSize > 0)
+				TOptional<double> Impact = FObjectNode::GetEstimatedMemoryImpact(Node);
+				if (Impact.IsSet())
 				{
-					TSharedPtr<FBaseTreeNode> ParentNode = Node.GetParent();
-					if (ParentNode.IsValid())
-					{
-						const int64 ParentEstimatedMemSize = FObjectNode::GetEstimatedMemorySize(*ParentNode);
-						if (ParentEstimatedMemSize > 0)
-						{
-							const double Ratio = double(EstimatedSize) / double(ParentEstimatedMemSize);
-							return FTableCellValue(static_cast<double>(Ratio));
-						}
-					}
+					return FTableCellValue(static_cast<double>(Impact.GetValue()));
 				}
 				return TOptional<FTableCellValue>();
 			}
@@ -1382,6 +1273,106 @@
 		AddColumn(ColumnRef);
 	}
 	//////////////////////////////////////////////////
+	// Inclusive System Memory Size Column
+	{
+		TSharedRef<FTableColumn> ColumnRef = MakeShared<FTableColumn>(FObjectTableColumns::InclSystemMemorySizeColumnId);
+		FTableColumn& Column = *ColumnRef;
+
+		Column.SetIndex(ColumnIndex++);
+
+		Column.SetShortName(LOCTEXT("InclSystemMemorySizeColumnName", "Incl. Mem Size"));
+		Column.SetTitleName(LOCTEXT("InclSystemMemorySizeColumnTitle", "Inclusive System Memory Size"));
+		Column.SetDescription(LOCTEXT("InclSystemMemorySizeColumnDesc", "The total estimated system memory size (bytes) of the UObject instance (including references)"));
+
+		Column.SetFlags(ETableColumnFlags::CanBeHidden |
+						ETableColumnFlags::CanBeFiltered);
+
+		Column.SetHorizontalAlignment(HAlign_Right);
+		Column.SetInitialWidth(100.0f);
+
+		Column.SetDataType(ETableCellDataType::Int64);
+
+		class FValueGetter : public FAggregatedValueGetter
+		{
+		public:
+			virtual const TOptional<FTableCellValue> GetNodeValue(const FBaseTreeNode& Node) const override
+			{
+				if (Node.Is<FObjectNode>())
+				{
+					const FObjectNode& ObjectNode = Node.As<FObjectNode>();
+					return FTableCellValue(int64(ObjectNode.GetInclusiveSystemMemorySize()));
+				}
+				return TOptional<FTableCellValue>();
+			}
+		};
+		TSharedRef<ITableCellValueGetter> Getter = MakeShared<FValueGetter>();
+		Column.SetValueGetter(Getter);
+
+		TSharedRef<ITableCellValueFormatter> Formatter = MakeShared<FInt64ValueFormatterAsMemory>();
+		Column.SetValueFormatter(Formatter);
+
+		TSharedRef<ITableCellValueSorter> Sorter = MakeShared<FSorterByInt64Value>(ColumnRef);
+		Column.SetValueSorter(Sorter);
+		Column.SetInitialSortMode(EColumnSortMode::Descending);
+
+		//TSharedRef<IFilterValueConverter> Converter = MakeShared<FMemoryFilterValueConverter>();
+		//Column.SetValueConverter(Converter);
+
+		Column.SetAggregation(ETableColumnAggregation::None);
+
+		AddColumn(ColumnRef);
+	}
+	//////////////////////////////////////////////////
+	// Inclusive Video Memory Size Column
+	{
+		TSharedRef<FTableColumn> ColumnRef = MakeShared<FTableColumn>(FObjectTableColumns::InclVideoMemorySizeColumnId);
+		FTableColumn& Column = *ColumnRef;
+
+		Column.SetIndex(ColumnIndex++);
+
+		Column.SetShortName(LOCTEXT("InclVideoMemorySizeColumnName", "Incl. Video Size"));
+		Column.SetTitleName(LOCTEXT("InclVideoMemorySizeColumnTitle", "Inclusive Video Memory Size"));
+		Column.SetDescription(LOCTEXT("InclVideoMemorySizeColumnDesc", "The total estimated video memory size (bytes) of the UObject instance (including references)"));
+
+		Column.SetFlags(ETableColumnFlags::CanBeHidden |
+						ETableColumnFlags::CanBeFiltered);
+
+		Column.SetHorizontalAlignment(HAlign_Right);
+		Column.SetInitialWidth(100.0f);
+
+		Column.SetDataType(ETableCellDataType::Int64);
+
+		class FValueGetter : public FAggregatedValueGetter
+		{
+		public:
+			virtual const TOptional<FTableCellValue> GetNodeValue(const FBaseTreeNode& Node) const override
+			{
+				if (Node.Is<FObjectNode>())
+				{
+					const FObjectNode& ObjectNode = Node.As<FObjectNode>();
+					return FTableCellValue(int64(ObjectNode.GetInclusiveVideoMemorySize()));
+				}
+				return TOptional<FTableCellValue>();
+			}
+		};
+		TSharedRef<ITableCellValueGetter> Getter = MakeShared<FValueGetter>();
+		Column.SetValueGetter(Getter);
+
+		TSharedRef<ITableCellValueFormatter> Formatter = MakeShared<FInt64ValueFormatterAsMemory>();
+		Column.SetValueFormatter(Formatter);
+
+		TSharedRef<ITableCellValueSorter> Sorter = MakeShared<FSorterByInt64Value>(ColumnRef);
+		Column.SetValueSorter(Sorter);
+		Column.SetInitialSortMode(EColumnSortMode::Descending);
+
+		//TSharedRef<IFilterValueConverter> Converter = MakeShared<FMemoryFilterValueConverter>();
+		//Column.SetValueConverter(Converter);
+
+		Column.SetAggregation(ETableColumnAggregation::None);
+
+		AddColumn(ColumnRef);
+	}
+	//////////////////////////////////////////////////
 	// Inclusive Estimated Size Column
 	{
 		TSharedRef<FTableColumn> ColumnRef = MakeShared<FTableColumn>(FObjectTableColumns::InclEstimatedSizeColumnId);
