[Insights] ObjectProfiler: Added drop down menu options for various filters:
* hide objects that are not referenced by any Actor
* hide objects with no Estimated Memory (System + Video)
* hide objects with Inclusive Estimated Memory (System + Video) < 1 MiB
* hide objects with Impact % < 1%
* shows only the UField objects (including UStruct objects)
* shows only the UStruct objects (including UClass and UFunction objects)
* shows only the UClass objects
* shows only the UFunction objects
* shows only the UPackage objects
* shows advanced UI (advanced filter configurator + hierarchy breadcrumb trail)

#Synced-CL 52248455

diff --git a/Engine/Source/Developer/TraceInsights/Private/Insights/ObjectProfiler/Widgets/SObjectTableTreeView.cpp b/Engine/Source/Developer/TraceInsights/Private/Insights/ObjectProfiler/Widgets/SObjectTableTreeView.cpp
--- a/Engine/Source/Developer/TraceInsights/Private/Insights/ObjectProfiler/Widgets/SObjectTableTreeView.cpp
+++ b/Engine/Source/Developer/TraceInsights/Private/Insights/ObjectProfiler/Widgets/SObjectTableTreeView.cpp
@@ -164,7 +164,6 @@ void SObjectTableTreeView::ConstructHeaderArea(TSharedRef<SVerticalBox> InHostBo
 							if (IsObjectDetailsViewVisible())
 							{
 								return FInsightsStyle::Get().GetBrush("Icons.FindPrevious.ToolBar");
-								
 							}
 							else
 							{
@@ -388,7 +387,247 @@ TSharedRef<SWidget> SObjectTableTreeView::MakeFiltersMenu()
 	FMenuBuilder MenuBuilder(/*bInShouldCloseWindowAfterMenuSelection=*/true, nullptr);
 	MenuBuilder.SetSearchable(false);
 
-	MenuBuilder.BeginSection("Filters"/*, LOCTEXT("ContextMenu_Section_Trackers", "Trackers")*/);
+	MenuBuilder.BeginSection("Filters", LOCTEXT("ContextMenu_Section_Filters", "Filters"));
+
+	MenuBuilder.AddMenuEntry(
+		LOCTEXT("HideObjectsWithZeroReferencingActors", "Hide objects with No Referencing Actors"),
+		LOCTEXT("HideObjectsWithZeroReferencingActors_TT", "Hides the objects that are not referenced by any Actor."),
+		FSlateIcon(),
+		FUIAction(
+			FExecuteAction::CreateLambda(
+				[this]
+				{
+					bHideObjectsWithZeroReferencingActors = !bHideObjectsWithZeroReferencingActors;
+					OnNodeFilteringChanged();
+				}),
+			FCanExecuteAction(),
+			FIsActionChecked::CreateLambda(
+				[this]() -> bool
+				{
+					return bHideObjectsWithZeroReferencingActors;
+				})),
+		NAME_None,
+		EUserInterfaceActionType::ToggleButton);
+
+	MenuBuilder.AddMenuEntry(
+		LOCTEXT("HideObjectsWithZeroEstimatedMemory", "Hide objects with No Estimated Memory"),
+		LOCTEXT("HideObjectsWithZeroEstimatedMemory_TT", "Hides the objects with no estimated memory (System + Video)."),
+		FSlateIcon(),
+		FUIAction(
+			FExecuteAction::CreateLambda(
+				[this]
+				{
+					bHideObjectsWithZeroEstimatedMemory = !bHideObjectsWithZeroEstimatedMemory;
+					OnNodeFilteringChanged();
+				}),
+			FCanExecuteAction(),
+			FIsActionChecked::CreateLambda(
+				[this]() -> bool
+				{
+					return bHideObjectsWithZeroEstimatedMemory;
+				})),
+		NAME_None,
+		EUserInterfaceActionType::ToggleButton);
+
+#if 0
+	MenuBuilder.AddMenuEntry(
+		LOCTEXT("HideObjectsWithLowInclusiveEstimatedMemory", "Hide objects with Low Inclusive Estimated Memory"),
+		LOCTEXT("HideObjectsWithLowInclusiveEstimatedMemory_TT", "Hides the objects with inclusive estimated memory (system + video) < 1 MiB."),
+		FSlateIcon(),
+		FUIAction(
+			FExecuteAction::CreateLambda(
+				[this]
+				{
+					bHideObjectsWithLowInclusiveEstimatedMemory = !bHideObjectsWithLowInclusiveEstimatedMemory;
+					OnNodeFilteringChanged();
+				}),
+			FCanExecuteAction(),
+			FIsActionChecked::CreateLambda(
+				[this]() -> bool
+				{
+					return bHideObjectsWithLowInclusiveEstimatedMemory;
+				})),
+		NAME_None,
+		EUserInterfaceActionType::ToggleButton);
+#endif
+
+	MenuBuilder.AddMenuEntry(
+		LOCTEXT("HideObjectsWithLowEstimatedMemory", "Hide objects with Low Estimated Memory"),
+		LOCTEXT("HideObjectsWithLowEstimatedMemory_TT", "Hides the objects with exclusive estimated memory (system + video) < 1 MiB."),
+		FSlateIcon(),
+		FUIAction(
+			FExecuteAction::CreateLambda(
+				[this]
+				{
+					bHideObjectsWithLowEstimatedMemory = !bHideObjectsWithLowEstimatedMemory;
+					OnNodeFilteringChanged();
+				}),
+			FCanExecuteAction(),
+			FIsActionChecked::CreateLambda(
+				[this]() -> bool
+				{
+					return bHideObjectsWithLowEstimatedMemory;
+				})),
+		NAME_None,
+		EUserInterfaceActionType::ToggleButton);
+
+	MenuBuilder.AddMenuEntry(
+		LOCTEXT("HideObjectsWithLowImpact", "Hide objects with Low Impact"),
+		LOCTEXT("HideObjectsWithLowImpact_TT", "Hides the objects with impact % < 1%."),
+		FSlateIcon(),
+		FUIAction(
+			FExecuteAction::CreateLambda(
+				[this]
+				{
+					bHideObjectsWithLowImpact = !bHideObjectsWithLowImpact;
+					OnNodeFilteringChanged();
+				}),
+			FCanExecuteAction(),
+			FIsActionChecked::CreateLambda(
+				[this]() -> bool
+				{
+					return bHideObjectsWithLowImpact;
+				})),
+		NAME_None,
+		EUserInterfaceActionType::ToggleButton);
+
+	MenuBuilder.AddMenuSeparator();
+
+	MenuBuilder.AddMenuEntry(
+		LOCTEXT("ShowOnlyFieldObjects", "Show only the UField objects"),
+		LOCTEXT("ShowOnlyFieldObjects_TT", "Shows only the UField objects (including UStruct objects)."),
+		FSlateIcon(),
+		FUIAction(
+			FExecuteAction::CreateLambda(
+				[this]
+				{
+					bShowOnlyFieldObjects = !bShowOnlyFieldObjects;
+					bFilterByClassType = bShowOnlyFieldObjects || bShowOnlyStructObjects || bShowOnlyClassObjects || bShowOnlyFunctionObjects || bShowOnlyPackageObjects;
+					OnNodeFilteringChanged();
+				}),
+			FCanExecuteAction(),
+			FIsActionChecked::CreateLambda(
+				[this]() -> bool
+				{
+					return bShowOnlyFieldObjects;
+				})),
+		NAME_None,
+		EUserInterfaceActionType::ToggleButton);
+
+	MenuBuilder.AddMenuEntry(
+		LOCTEXT("ShowOnlyStructObjects", "Show only the UStruct objects"),
+		LOCTEXT("ShowOnlyStructObjects_TT", "Shows only the UStruct objects (including UClass and UFunction objects)."),
+		FSlateIcon(),
+		FUIAction(
+			FExecuteAction::CreateLambda(
+				[this]
+				{
+					bShowOnlyStructObjects = !bShowOnlyStructObjects;
+					bFilterByClassType = bShowOnlyFieldObjects || bShowOnlyStructObjects || bShowOnlyClassObjects || bShowOnlyFunctionObjects || bShowOnlyPackageObjects;
+					OnNodeFilteringChanged();
+				}),
+			FCanExecuteAction(),
+			FIsActionChecked::CreateLambda(
+				[this]() -> bool
+				{
+					return bShowOnlyStructObjects;
+				})),
+		NAME_None,
+		EUserInterfaceActionType::ToggleButton);
+
+	MenuBuilder.AddMenuEntry(
+		LOCTEXT("ShowOnlyClassObjects", "Show only the UClass objects"),
+		LOCTEXT("ShowOnlyClassObjects_TT", "Shows only the UClass objects."),
+		FSlateIcon(),
+		FUIAction(
+			FExecuteAction::CreateLambda(
+				[this]
+				{
+					bShowOnlyClassObjects = !bShowOnlyClassObjects;
+					bFilterByClassType = bShowOnlyFieldObjects || bShowOnlyStructObjects || bShowOnlyClassObjects || bShowOnlyFunctionObjects || bShowOnlyPackageObjects;
+					OnNodeFilteringChanged();
+				}),
+			FCanExecuteAction(),
+			FIsActionChecked::CreateLambda(
+				[this]() -> bool
+				{
+					return bShowOnlyClassObjects;
+				})),
+		NAME_None,
+		EUserInterfaceActionType::ToggleButton);
+
+	MenuBuilder.AddMenuEntry(
+		LOCTEXT("ShowOnlyFunctionObjects", "Show only the UFunction objects"),
+		LOCTEXT("ShowOnlyFunctionObjects_TT", "Shows only the UFunction objects."),
+		FSlateIcon(),
+		FUIAction(
+			FExecuteAction::CreateLambda(
+				[this]
+				{
+					bShowOnlyFunctionObjects = !bShowOnlyFunctionObjects;
+					bFilterByClassType = bShowOnlyFieldObjects || bShowOnlyStructObjects || bShowOnlyClassObjects || bShowOnlyFunctionObjects || bShowOnlyPackageObjects;
+					OnNodeFilteringChanged();
+				}),
+			FCanExecuteAction(),
+			FIsActionChecked::CreateLambda(
+				[this]() -> bool
+				{
+					return bShowOnlyFunctionObjects;
+				})),
+		NAME_None,
+		EUserInterfaceActionType::ToggleButton);
+
+	MenuBuilder.AddMenuEntry(
+		LOCTEXT("ShowOnlyPackageObjects", "Show only the UPackage objects"),
+		LOCTEXT("ShowOnlyPackageObjects_TT", "Shows only the UPackage objects."),
+		FSlateIcon(),
+		FUIAction(
+			FExecuteAction::CreateLambda(
+				[this]
+				{
+					bShowOnlyPackageObjects = !bShowOnlyPackageObjects;
+					bFilterByClassType = bShowOnlyFieldObjects || bShowOnlyStructObjects || bShowOnlyClassObjects || bShowOnlyFunctionObjects || bShowOnlyPackageObjects;
+					OnNodeFilteringChanged();
+				}),
+			FCanExecuteAction(),
+			FIsActionChecked::CreateLambda(
+				[this]() -> bool
+				{
+					return bShowOnlyPackageObjects;
+				})),
+		NAME_None,
+		EUserInterfaceActionType::ToggleButton);
+
+	MenuBuilder.AddMenuSeparator();
+
+	MenuBuilder.AddMenuEntry(
+		LOCTEXT("ShowAdvancedUI", "Show Advanced UI"),
+		LOCTEXT("ShowAdvancedUI_TT", "Shows the advanced UI options (ex.: the advanced filter configurator and the grouping breadcrumb trail)."),
+		FSlateIcon(),
+		FUIAction(
+			FExecuteAction::CreateLambda(
+				[this]
+				{
+					if (bIsAdvancedFilterConfiguratorVisible || bIsHierarchyBreadcrumbTrailVisible)
+					{
+						bIsAdvancedFilterConfiguratorVisible = false;
+						bIsHierarchyBreadcrumbTrailVisible = false;
+					}
+					else
+					{
+						bIsAdvancedFilterConfiguratorVisible = true;
+						bIsHierarchyBreadcrumbTrailVisible = true;
+					}
+					OnNodeFilteringChanged();
+				}),
+			FCanExecuteAction(),
+			FIsActionChecked::CreateLambda(
+				[this]() -> bool
+				{
+					return bIsAdvancedFilterConfiguratorVisible || bIsHierarchyBreadcrumbTrailVisible;
+				})),
+		NAME_None,
+		EUserInterfaceActionType::ToggleButton);
 
 	MenuBuilder.EndSection();
 
@@ -703,12 +942,12 @@ void SObjectTableTreeView::ExtendMenuAfterMisc(FMenuBuilder& MenuBuilder)
 
 bool SObjectTableTreeView::HasCustomNodeFilter() const
 {
-	return bHideObjectsWithZeroStructureSize
+	return bFilterByClassType
+		|| bHideObjectsWithZeroReferencingActors
 		|| bHideObjectsWithZeroEstimatedMemory
 		|| bHideObjectsWithLowInclusiveEstimatedMemory
 		|| bHideObjectsWithLowEstimatedMemory
-		|| bHideObjectsWithLowImpact
-		|| bHideObjectsWithZeroReferencingActors;
+		|| bHideObjectsWithLowImpact;
 }
 
 ////////////////////////////////////////////////////////////////////////////////////////////////////
@@ -718,16 +957,48 @@ bool SObjectTableTreeView::FilterNodeCustom(const FTableTreeNode& InNode) const
 	if (InNode.Is<FObjectNode>())
 	{
 		const FObjectNode& ObjectNode = InNode.As<FObjectNode>();
-		if (bHideObjectsWithZeroStructureSize)
+		if (bFilterByClassType)
 		{
-			if (ObjectNode.GetStructureSize() == 0)
+			if (bShowOnlyFieldObjects)
 			{
-				return false;
+				if (!ObjectNode.Is<FFieldObjectNode>())
+				{
+					return false;
+				}
+			}
+			if (bShowOnlyStructObjects)
+			{
+				if (!ObjectNode.Is<FStructObjectNode>())
+				{
+					return false;
+				}
+			}
+			if (bShowOnlyClassObjects)
+			{
+				if (!ObjectNode.Is<FClassObjectNode>())
+				{
+					return false;
+				}
+			}
+			if (bShowOnlyFunctionObjects)
+			{
+				if (!ObjectNode.Is<FFunctionObjectNode>())
+				{
+					return false;
+				}
+			}
+			if (bShowOnlyPackageObjects)
+			{
+				if (!ObjectNode.Is<FPackageObjectNode>())
+				{
+					return false;
+				}
 			}
 		}
-		if (bHideObjectsWithZeroEstimatedMemory)
+		if (bHideObjectsWithZeroReferencingActors)
 		{
-			if (ObjectNode.GetEstimatedMemorySize() == 0)
+			const uint32 NumReferences = ObjectNode.GetNumReferences();
+			if (NumReferences == 0)
 			{
 				return false;
 			}
@@ -740,6 +1011,14 @@ bool SObjectTableTreeView::FilterNodeCustom(const FTableTreeNode& InNode) const
 				return false;
 			}
 		}
+		if (bHideObjectsWithZeroEstimatedMemory)
+		{
+			const int64 EstimatedMemSize = ObjectNode.GetEstimatedMemorySize();
+			if (EstimatedMemSize == 0)
+			{
+				return false;
+			}
+		}
 		if (bHideObjectsWithLowEstimatedMemory)
 		{
 			const int64 EstimatedMemSize = ObjectNode.GetEstimatedMemorySize();
@@ -750,21 +1029,51 @@ bool SObjectTableTreeView::FilterNodeCustom(const FTableTreeNode& InNode) const
 		}
 		if (bHideObjectsWithLowImpact)
 		{
-			const double EstimatedImpact = ObjectNode.GetEstimatedMemoryImpact();
-			if (EstimatedImpact < 0.01)
+			TOptional<double> EstimatedImpact = ObjectNode.GetEstimatedMemoryImpact();
+			if (!EstimatedImpact.IsSet() || EstimatedImpact.GetValue() < 0.01)
 			{
 				return false;
 			}
 		}
+	}
+#if 0
+	else
+	{
 		if (bHideObjectsWithZeroReferencingActors)
 		{
-			const uint32 NumReferences = ObjectNode.GetNumReferences();
+			const uint32 NumReferences = FObjectNode::GetNumReferences(InNode);
 			if (NumReferences == 0)
 			{
 				return false;
 			}
 		}
+		if (bHideObjectsWithZeroEstimatedMemory)
+		{
+			const int64 EstimatedMemSize = FObjectNode::GetEstimatedMemorySize(InNode);
+			if (EstimatedMemSize == 0)
+			{
+				return false;
+			}
+		}
+		if (bHideObjectsWithLowEstimatedMemory)
+		{
+			const int64 EstimatedMemSize = FObjectNode::GetEstimatedMemorySize(InNode);
+			if (EstimatedMemSize < 1024 * 1024)
+			{
+				return false;
+			}
+		}
+		if (bHideObjectsWithLowImpact)
+		{
+			TOptional<double> EstimatedImpact = FObjectNode::GetEstimatedMemoryImpact(InNode);
+			if (!EstimatedImpact.IsSet() || EstimatedImpact.GetValue() < 0.01)
+			{
+				return false;
+			}
+		}
 	}
+#endif
+
 	return true;
 }
 
diff --git a/Engine/Source/Developer/TraceInsights/Private/Insights/ObjectProfiler/Widgets/SObjectTableTreeView.h b/Engine/Source/Developer/TraceInsights/Private/Insights/ObjectProfiler/Widgets/SObjectTableTreeView.h
--- a/Engine/Source/Developer/TraceInsights/Private/Insights/ObjectProfiler/Widgets/SObjectTableTreeView.h
+++ b/Engine/Source/Developer/TraceInsights/Private/Insights/ObjectProfiler/Widgets/SObjectTableTreeView.h
@@ -121,15 +121,21 @@ private:
 	//////////////////////////////////////////////////
 	// Filters
 
-	bool bHideObjectsWithZeroStructureSize = false; // only show UStruct objects (including UClass and UFunction objects; with valid Structure size)
-	bool bHideObjectsWithZeroEstimatedMemory = false; // hide objects with no Estimated Memory (System + Video)
+	bool bHideObjectsWithZeroReferencingActors = false; // hide objects that are not referenced by any Actor
+	bool bHideObjectsWithZeroEstimatedMemory = true; // hide objects with no Estimated Memory (System + Video)
 	bool bHideObjectsWithLowInclusiveEstimatedMemory = false; // hide objects with Inclusive Estimated Memory (System + Video) < 1 MiB
 	bool bHideObjectsWithLowEstimatedMemory = false; // hide objects with Exclusive Estimated Memory (System + Video) < 1 MiB
 	bool bHideObjectsWithLowImpact = false; // hide objects with Impact % < 1%
-	bool bHideObjectsWithZeroReferencingActors = false; // hide objects that are not referenced by any Actor
 
-	bool bIsAdvancedFilterConfiguratorVisible = true;
-	bool bIsHierarchyBreadcrumbTrailVisible = true;
+	bool bFilterByClassType = false;
+	bool bShowOnlyFieldObjects = false; // shows only the UField objects (including UStruct objects)
+	bool bShowOnlyStructObjects = false; // shows only the UStruct objects (including UClass and UFunction objects)
+	bool bShowOnlyClassObjects = false; // shows only the UClass objects
+	bool bShowOnlyFunctionObjects = false; // shows only the UFunction objects
+	bool bShowOnlyPackageObjects = false; // shows only the UPackage objects
+
+	bool bIsAdvancedFilterConfiguratorVisible = false;
+	bool bIsHierarchyBreadcrumbTrailVisible = false;
 
 	//////////////////////////////////////////////////
 
