Fix ensure when matching object to actors in Memory View.

#Synced-CL 51753244

diff --git a/Engine/Source/Developer/TraceInsights/Private/Insights/ObjectProfiler/Common/ObjectToActorResolver.cpp b/Engine/Source/Developer/TraceInsights/Private/Insights/ObjectProfiler/Common/ObjectToActorResolver.cpp
--- a/Engine/Source/Developer/TraceInsights/Private/Insights/ObjectProfiler/Common/ObjectToActorResolver.cpp
+++ b/Engine/Source/Developer/TraceInsights/Private/Insights/ObjectProfiler/Common/ObjectToActorResolver.cpp
@@ -349,7 +349,7 @@ void FBuildPackageDependencyMapBuilder::PerformBatchedReferenceChainSearch(TSet<
 	}
 
 	// Complete any remaining progress
-	const float RemainingWork = ChainProcessingTask.TotalAmountOfWork - ChainProcessingTask.CompletedWork;
+	const float RemainingWork = ChainProcessingTask.TotalAmountOfWork - ChainProcessingTask.CompletedWork - ChainProcessingTask.CurrentFrameScope;
 	if (RemainingWork > 0.0f)
 	{
 		ChainProcessingTask.EnterProgressFrame(RemainingWork);
