[Insights] Object Profiler: Created the IAssetInfoProfiler for both standalone and editor use cases.

#rnx

#Synced-CL 52250744

diff --git a/Engine/Source/Developer/TraceInsights/Private/Insights/ObjectProfiler/AssetInfoProvider.cpp b/Engine/Source/Developer/TraceInsights/Private/Insights/ObjectProfiler/AssetInfoProvider.cpp
--- a/Engine/Source/Developer/TraceInsights/Private/Insights/ObjectProfiler/AssetInfoProvider.cpp
+++ b/Engine/Source/Developer/TraceInsights/Private/Insights/ObjectProfiler/AssetInfoProvider.cpp
@@ -70,7 +70,7 @@ FAssetInfoProvider::~FAssetInfoProvider()
 
 void FAssetInfoProvider::InitCategories()
 {
-#define UE_INSIGHTS_ENUM(_Category, _ParentCategory, _Color, _DisplayName) \
+#define UE_INSIGHTS_ENUM(_ParentCategory, _Category, _Color, _DisplayName) \
 	{ \
 		FAssetInfoCategory& CategoryInfo = Categories.AddDefaulted_GetRef(); \
 		CategoryInfo.Name = FName(#_Category); \
diff --git a/Engine/Source/Developer/TraceInsights/Private/Insights/ObjectProfiler/ObjectProfilerManager.cpp b/Engine/Source/Developer/TraceInsights/Private/Insights/ObjectProfiler/ObjectProfilerManager.cpp
--- a/Engine/Source/Developer/TraceInsights/Private/Insights/ObjectProfiler/ObjectProfilerManager.cpp
+++ b/Engine/Source/Developer/TraceInsights/Private/Insights/ObjectProfiler/ObjectProfilerManager.cpp
@@ -17,6 +17,9 @@
 #include "Insights/Common/InsightsMenuBuilder.h"
 #include "Insights/InsightsManager.h"
 #include "Insights/InsightsStyle.h"
+#include "Insights/ObjectProfiler/AssetInfoProvider.h"
+#include "Insights/ObjectProfiler/EditorAssetInfoProvider.h"
+#include "Insights/ObjectProfiler/IAssetInfoProvider.h"
 #include "Insights/ObjectProfiler/ViewModels/ObjectTable.h"
 #include "Insights/ObjectProfiler/Widgets/SObjectProfilerWindow.h"
 
@@ -197,6 +200,17 @@ TSharedRef<SDockTab> FObjectProfilerManager::SpawnTab(const FSpawnTabArgs& Args)
 
 	AssignProfilerWindow(Window);
 
+	if (!AssetInfoProvider)
+	{
+#if WITH_EDITOR
+		AssetInfoProvider = MakeShared<FEditorAssetInfoProvider>();
+#else
+		AssetInfoProvider = MakeShared<FAssetInfoProvider>();
+#endif
+	}
+
+	Window->SetAssetInfoProvider(AssetInfoProvider);
+
 	return DockTab;
 }
 
diff --git a/Engine/Source/Developer/TraceInsights/Private/Insights/ObjectProfiler/ObjectProfilerManager.h b/Engine/Source/Developer/TraceInsights/Private/Insights/ObjectProfiler/ObjectProfilerManager.h
--- a/Engine/Source/Developer/TraceInsights/Private/Insights/ObjectProfiler/ObjectProfilerManager.h
+++ b/Engine/Source/Developer/TraceInsights/Private/Insights/ObjectProfiler/ObjectProfilerManager.h
@@ -19,6 +19,7 @@
 namespace UE::Insights::ObjectProfiler
 {
 
+class IAssetInfoProvider;
 class SObjectProfilerWindow;
 
 DECLARE_LOG_CATEGORY_EXTERN(LogObjectProfiler, Log, All);
@@ -156,6 +157,9 @@ private:
 	/** A weak pointer to the Object Insights window. */
 	TWeakPtr<SObjectProfilerWindow> ProfilerWindowWeakPtr;
 
+	/** The Asset Info Provider to be used by the Object Insights window. */
+	TSharedPtr<IAssetInfoProvider> AssetInfoProvider;
+
 	/** If the Objects table/tree view is visible or hidden. */
 	bool bIsObjectTableTreeViewVisible = false;
 
