[Insights] ObjectProfiler: Enabled sorters to allow mixing group and leaf nodes.

#rnx
#Synced-CL 52393202

diff --git a/Engine/Source/Developer/TraceInsights/Private/Insights/ObjectProfiler/ViewModels/ObjectTable.cpp b/Engine/Source/Developer/TraceInsights/Private/Insights/ObjectProfiler/ViewModels/ObjectTable.cpp
--- a/Engine/Source/Developer/TraceInsights/Private/Insights/ObjectProfiler/ViewModels/ObjectTable.cpp
+++ b/Engine/Source/Developer/TraceInsights/Private/Insights/ObjectProfiler/ViewModels/ObjectTable.cpp
@@ -999,7 +999,8 @@ void FObjectTable::AddDefaultColumns()
 		TSharedRef<ITableCellValueFormatter> Formatter = MakeShared<FInt64ValueFormatterAsMemory>();
 		Column.SetValueFormatter(Formatter);
 
-		TSharedRef<ITableCellValueSorter> Sorter = MakeShared<FSorterByInt64Value>(ColumnRef);
+		TSharedRef<FSorterByInt64Value> Sorter = MakeShared<FSorterByInt64Value>(ColumnRef);
+		Sorter->SetShouldSortGroupsFirst(false);
 		Column.SetValueSorter(Sorter);
 		Column.SetInitialSortMode(EColumnSortMode::Descending);
 
@@ -1050,7 +1051,8 @@ void FObjectTable::AddDefaultColumns()
 		TSharedRef<ITableCellValueFormatter> Formatter = MakeShared<FInt64ValueFormatterAsMemory>();
 		Column.SetValueFormatter(Formatter);
 
-		TSharedRef<ITableCellValueSorter> Sorter = MakeShared<FSorterByInt64Value>(ColumnRef);
+		TSharedRef<FSorterByInt64Value> Sorter = MakeShared<FSorterByInt64Value>(ColumnRef);
+		Sorter->SetShouldSortGroupsFirst(false);
 		Column.SetValueSorter(Sorter);
 		Column.SetInitialSortMode(EColumnSortMode::Descending);
 
@@ -1101,7 +1103,8 @@ void FObjectTable::AddDefaultColumns()
 		TSharedRef<ITableCellValueFormatter> Formatter = MakeShared<FInt64ValueFormatterAsMemory>();
 		Column.SetValueFormatter(Formatter);
 
-		TSharedRef<ITableCellValueSorter> Sorter = MakeShared<FSorterByInt64Value>(ColumnRef);
+		TSharedRef<FSorterByInt64Value> Sorter = MakeShared<FSorterByInt64Value>(ColumnRef);
+		Sorter->SetShouldSortGroupsFirst(false);
 		Column.SetValueSorter(Sorter);
 		Column.SetInitialSortMode(EColumnSortMode::Descending);
 
@@ -1191,7 +1194,8 @@ void FObjectTable::AddDefaultColumns()
 		TSharedRef<ITableCellValueFormatter> Formatter = MakeShared<FEstimatedSizeFormatter>();
 		Column.SetValueFormatter(Formatter);
 
-		TSharedRef<ITableCellValueSorter> Sorter = MakeShared<FSorterByInt64Value>(ColumnRef);
+		TSharedRef<FSorterByInt64Value> Sorter = MakeShared<FSorterByInt64Value>(ColumnRef);
+		Sorter->SetShouldSortGroupsFirst(false);
 		Column.SetValueSorter(Sorter);
 		Column.SetInitialSortMode(EColumnSortMode::Descending);
 
@@ -1264,7 +1268,8 @@ void FObjectTable::AddDefaultColumns()
 		TSharedRef<ITableCellValueFormatter> Formatter = MakeShared<FImpactFormatter>();
 		Column.SetValueFormatter(Formatter);
 
-		TSharedRef<ITableCellValueSorter> Sorter = MakeShared<FSorterByDoubleValue>(ColumnRef);
+		TSharedRef<FSorterByDoubleValue> Sorter = MakeShared<FSorterByDoubleValue>(ColumnRef);
+		Sorter->SetShouldSortGroupsFirst(false);
 		Column.SetValueSorter(Sorter);
 		Column.SetInitialSortMode(EColumnSortMode::Descending);
 
@@ -1311,7 +1316,8 @@ void FObjectTable::AddDefaultColumns()
 		TSharedRef<ITableCellValueFormatter> Formatter = MakeShared<FInt64ValueFormatterAsMemory>();
 		Column.SetValueFormatter(Formatter);
 
-		TSharedRef<ITableCellValueSorter> Sorter = MakeShared<FSorterByInt64Value>(ColumnRef);
+		TSharedRef<FSorterByInt64Value> Sorter = MakeShared<FSorterByInt64Value>(ColumnRef);
+		Sorter->SetShouldSortGroupsFirst(false);
 		Column.SetValueSorter(Sorter);
 		Column.SetInitialSortMode(EColumnSortMode::Descending);
 
@@ -1361,7 +1367,8 @@ void FObjectTable::AddDefaultColumns()
 		TSharedRef<ITableCellValueFormatter> Formatter = MakeShared<FInt64ValueFormatterAsMemory>();
 		Column.SetValueFormatter(Formatter);
 
-		TSharedRef<ITableCellValueSorter> Sorter = MakeShared<FSorterByInt64Value>(ColumnRef);
+		TSharedRef<FSorterByInt64Value> Sorter = MakeShared<FSorterByInt64Value>(ColumnRef);
+		Sorter->SetShouldSortGroupsFirst(false);
 		Column.SetValueSorter(Sorter);
 		Column.SetInitialSortMode(EColumnSortMode::Descending);
 
@@ -1451,7 +1458,8 @@ void FObjectTable::AddDefaultColumns()
 		TSharedRef<ITableCellValueFormatter> Formatter = MakeShared<FEstimatedSizeFormatter>();
 		Column.SetValueFormatter(Formatter);
 
-		TSharedRef<ITableCellValueSorter> Sorter = MakeShared<FSorterByInt64Value>(ColumnRef);
+		TSharedRef<FSorterByInt64Value> Sorter = MakeShared<FSorterByInt64Value>(ColumnRef);
+		Sorter->SetShouldSortGroupsFirst(false);
 		Column.SetValueSorter(Sorter);
 		Column.SetInitialSortMode(EColumnSortMode::Descending);
 
@@ -1499,7 +1507,8 @@ void FObjectTable::AddDefaultColumns()
 		TSharedRef<ITableCellValueFormatter> Formatter = MakeShared<FInt64ValueFormatterAsNumber>();
 		Column.SetValueFormatter(Formatter);
 
-		TSharedRef<ITableCellValueSorter> Sorter = MakeShared<FSorterByInt64Value>(ColumnRef);
+		TSharedRef<FSorterByInt64Value> Sorter = MakeShared<FSorterByInt64Value>(ColumnRef);
+		Sorter->SetShouldSortGroupsFirst(false);
 		Column.SetValueSorter(Sorter);
 		Column.SetInitialSortMode(EColumnSortMode::Descending);
 
