Add save to file and load from file functionality to Memory View. Fixed issues encountered when taking multiple snapshots or when loading a trace in Memory View. Fixed some more small code quality issues.

#Synced-CL 51753244

diff --git a/Engine/Source/Developer/TraceInsights/Private/Insights/ObjectProfiler/ObjectProfilerManager.cpp b/Engine/Source/Developer/TraceInsights/Private/Insights/ObjectProfiler/ObjectProfilerManager.cpp
--- a/Engine/Source/Developer/TraceInsights/Private/Insights/ObjectProfiler/ObjectProfilerManager.cpp
+++ b/Engine/Source/Developer/TraceInsights/Private/Insights/ObjectProfiler/ObjectProfilerManager.cpp
@@ -68,7 +68,7 @@ TSharedPtr<FObjectProfilerManager> FObjectProfilerManager::CreateInstance()
 FObjectProfilerManager::FObjectProfilerManager(TSharedRef<FUICommandList> InCommandList)
 	: CommandList(InCommandList)
 	, ActionManager(this)
-	, MajorTabId(SObjectProfilerWindow::ObjectProfilerTabId)
+	, MajorTabId(ObjectProfilerTabId)
 	, LogListingName("ObjectInsights")
 {
 }
diff --git a/Engine/Source/Developer/TraceInsights/Private/Insights/ObjectProfiler/Widgets/SObjectProfilerWindow.cpp b/Engine/Source/Developer/TraceInsights/Private/Insights/ObjectProfiler/Widgets/SObjectProfilerWindow.cpp
--- a/Engine/Source/Developer/TraceInsights/Private/Insights/ObjectProfiler/Widgets/SObjectProfilerWindow.cpp
+++ b/Engine/Source/Developer/TraceInsights/Private/Insights/ObjectProfiler/Widgets/SObjectProfilerWindow.cpp
@@ -29,7 +29,7 @@ namespace UE::Insights::ObjectProfiler
 
 const FName FObjectProfilerTabs::ObjectTableTreeViewId("ObjectTableTreeView");
 const FName FObjectProfilerTabs::ObjectDetailsViewId("ObjectDetailsView");
-const FName SObjectProfilerWindow::ObjectProfilerTabId("ObjectProfiler");
+const FName ObjectProfilerTabId("ObjectProfiler");
 const FName ObjectProfilerExtenderFeatureName("ObjectProfilerExtender");
 
 ////////////////////////////////////////////////////////////////////////////////////////////////////
diff --git a/Engine/Source/Developer/TraceInsights/Private/Insights/ObjectProfiler/Widgets/SObjectProfilerWindow.h b/Engine/Source/Developer/TraceInsights/Private/Insights/ObjectProfiler/Widgets/SObjectProfilerWindow.h
--- a/Engine/Source/Developer/TraceInsights/Private/Insights/ObjectProfiler/Widgets/SObjectProfilerWindow.h
+++ b/Engine/Source/Developer/TraceInsights/Private/Insights/ObjectProfiler/Widgets/SObjectProfilerWindow.h
@@ -31,6 +31,8 @@ struct FObjectProfilerTabs
 	static const FName ObjectDetailsViewId;
 };
 
+extern TRACEINSIGHTS_API const FName ObjectProfilerTabId;
+
 ////////////////////////////////////////////////////////////////////////////////////////////////////
 
 class FCustomBarGraphSegment : public UE::Insights::IBarGraphSegment
@@ -61,7 +63,6 @@ private:
 class SObjectProfilerWindow : public ::Insights::SMajorTabWindow, public IObjectProfilerMemoryView
 {
 public:
-	static const FName ObjectProfilerTabId;
 
 	/** Default constructor. */
 	SObjectProfilerWindow();
diff --git a/Engine/Source/Developer/TraceInsights/Private/Insights/ObjectProfiler/Widgets/SObjectTableTreeView.cpp b/Engine/Source/Developer/TraceInsights/Private/Insights/ObjectProfiler/Widgets/SObjectTableTreeView.cpp
--- a/Engine/Source/Developer/TraceInsights/Private/Insights/ObjectProfiler/Widgets/SObjectTableTreeView.cpp
+++ b/Engine/Source/Developer/TraceInsights/Private/Insights/ObjectProfiler/Widgets/SObjectTableTreeView.cpp
@@ -677,6 +677,8 @@ void SObjectTableTreeView::InitAvailableViewPresets()
 
 void SObjectTableTreeView::Reset()
 {
+	CurrentSnapshotId = uint32(-1);
+	NextTimestamp = 0;
 	STableTreeView::Reset();
 }
 
@@ -688,7 +690,6 @@ void SObjectTableTreeView::Tick(const FGeometry& AllottedGeometry, const double
 
 	// We need to check if the list of UObject snapshots has changed.
 	// But, ensure we do not check too often.
-	static uint64 NextTimestamp = 0;
 	uint64 Time = FPlatformTime::Cycles64();
 	if (Time > NextTimestamp && !bIsUpdateRunning)
 	{
diff --git a/Engine/Source/Developer/TraceInsights/Private/Insights/ObjectProfiler/Widgets/SObjectTableTreeView.h b/Engine/Source/Developer/TraceInsights/Private/Insights/ObjectProfiler/Widgets/SObjectTableTreeView.h
--- a/Engine/Source/Developer/TraceInsights/Private/Insights/ObjectProfiler/Widgets/SObjectTableTreeView.h
+++ b/Engine/Source/Developer/TraceInsights/Private/Insights/ObjectProfiler/Widgets/SObjectTableTreeView.h
@@ -147,6 +147,8 @@ private:
 	FText SelectionStructSizeText;
 	FText SelectionSystemMemSizeText;
 	FText SelectionVideoMemSizeText;
+
+	uint64 NextTimestamp = 0;
 };
 
 ////////////////////////////////////////////////////////////////////////////////////////////////////
diff --git a/Engine/Source/Developer/TraceInsights/Private/Insights/TraceInsightsModule.cpp b/Engine/Source/Developer/TraceInsights/Private/Insights/TraceInsightsModule.cpp
--- a/Engine/Source/Developer/TraceInsights/Private/Insights/TraceInsightsModule.cpp
+++ b/Engine/Source/Developer/TraceInsights/Private/Insights/TraceInsightsModule.cpp
@@ -365,7 +365,7 @@ void FTraceInsightsModule::AddAreaForSessionViewer(TSharedRef<FTabManager::FLayo
 	Stack->AddTab(FInsightsManagerTabs::LoadingProfilerTabId, ETabState::ClosedTab);
 	Stack->AddTab(FInsightsManagerTabs::NetworkingProfilerTabId, ETabState::ClosedTab);
 	Stack->AddTab(FInsightsManagerTabs::MemoryProfilerTabId, ETabState::ClosedTab);
-	Stack->AddTab(UE::Insights::ObjectProfiler::SObjectProfilerWindow::ObjectProfilerTabId, ETabState::ClosedTab);
+	Stack->AddTab(UE::Insights::ObjectProfiler::ObjectProfilerTabId, ETabState::ClosedTab);
 	//Stack->SetForegroundTab(FTabId(FInsightsManagerTabs::TimingProfilerTabId));
 
 	// Create area for the main window.
@@ -382,7 +382,7 @@ void FTraceInsightsModule::AddAreaForSessionViewer(TSharedRef<FTabManager::FLayo
 	Stack->AddTab(FInsightsManagerTabs::NetworkingProfilerTabId, ETabState::ClosedTab);
 	Stack->AddTab(FInsightsManagerTabs::MemoryProfilerTabId, ETabState::ClosedTab);
 	Stack->AddTab(FInsightsManagerTabs::MessageLogTabId, ETabState::ClosedTab);
-	Stack->AddTab(UE::Insights::ObjectProfiler::SObjectProfilerWindow::ObjectProfilerTabId, ETabState::ClosedTab);
+	Stack->AddTab(UE::Insights::ObjectProfiler::ObjectProfilerTabId, ETabState::ClosedTab);
 	Stack->SetForegroundTab(FTabId(FInsightsManagerTabs::TimingProfilerTabId));
 
 	Layout->AddArea
diff --git a/Engine/Source/Developer/TraceInsights/Public/Insights/ObjectProfiler/IObjectProfilerExtender.h b/Engine/Source/Developer/TraceInsights/Public/Insights/ObjectProfiler/IObjectProfilerExtender.h
--- a/Engine/Source/Developer/TraceInsights/Public/Insights/ObjectProfiler/IObjectProfilerExtender.h
+++ b/Engine/Source/Developer/TraceInsights/Public/Insights/ObjectProfiler/IObjectProfilerExtender.h
@@ -14,6 +14,7 @@ namespace UE::Insights::ObjectProfiler
 {
 
 extern TRACEINSIGHTS_API const FName ObjectProfilerExtenderFeatureName;
+extern TRACEINSIGHTS_API const FName ObjectProfilerTabId;
 
 struct FObjectProfilerExtenderTickParams
 {
