[Insights] ObjectProfiler: Added grouping by object name.

#Synced-CL 52443010

diff --git a/Engine/Source/Developer/TraceInsights/Private/Insights/ObjectProfiler/ViewModels/ObjectGroupingByObjectName.cpp b/Engine/Source/Developer/TraceInsights/Private/Insights/ObjectProfiler/ViewModels/ObjectGroupingByObjectName.cpp
new file mode 100644
--- /dev/null
+++ b/Engine/Source/Developer/TraceInsights/Private/Insights/ObjectProfiler/ViewModels/ObjectGroupingByObjectName.cpp
@@ -0,0 +1,124 @@
+// Copyright Epic Games, Inc. All Rights Reserved.
+
+#include "ObjectGroupingByObjectName.h"
+
+#include "Internationalization/Internationalization.h"
+
+// TraceInsightsCore
+#include "InsightsCore/Common/AsyncOperationProgress.h"
+
+// TraceInsights
+#include "Insights/ObjectProfiler/ViewModels/ObjectNode.h"
+
+#define LOCTEXT_NAMESPACE "UE::Insights::ObjectProfiler::FObjectNode"
+
+namespace UE::Insights::ObjectProfiler
+{
+
+INSIGHTS_IMPLEMENT_RTTI(FObjectGroupingByObjectName)
+
+////////////////////////////////////////////////////////////////////////////////////////////////////
+
+FObjectGroupingByObjectName::FObjectGroupingByObjectName()
+	: FTreeNodeGrouping(
+		LOCTEXT("Grouping_ByObjectName_ShortName", "Name"),
+		LOCTEXT("Grouping_ByObjectName_TitleName", "By Object Name"),
+		LOCTEXT("Grouping_ByObjectName_Desc", "Group together objects with same name."),
+		nullptr)
+{
+	SetColor(FLinearColor(1.0f, 0.45f, 0.6f, 1.0f));
+}
+
+////////////////////////////////////////////////////////////////////////////////////////////////////
+
+FObjectGroupingByObjectName::~FObjectGroupingByObjectName()
+{
+}
+
+////////////////////////////////////////////////////////////////////////////////////////////////////
+
+void FObjectGroupingByObjectName::GroupNodes(
+	const TArray<FTableTreeNodePtr>& Nodes,
+	FTableTreeNode& ParentGroup,
+	TWeakPtr<FTable> InParentTable,
+	IAsyncOperationProgress& InAsyncOperationProgress) const
+{
+	ParentGroup.ClearChildren();
+
+	TArray<TSharedPtr<FObjectNode>> ObjectNodes;
+
+	for (FTableTreeNodePtr NodePtr : Nodes)
+	{
+		if (InAsyncOperationProgress.ShouldCancelAsyncOp())
+		{
+			break;
+		}
+
+		if (!NodePtr->Is<FObjectNode>())
+		{
+			ParentGroup.AddChildAndSetParent(NodePtr);
+			continue;
+		}
+
+		FObjectNode& ObjNode = NodePtr->As<FObjectNode>();
+		if (ObjNode.GetObjectId() == uint32(-1))
+		{
+			ParentGroup.AddChildAndSetParent(NodePtr);
+			continue;
+		}
+
+		ObjectNodes.Add(StaticCastSharedPtr<FObjectNode>(NodePtr));
+	}
+
+	ObjectNodes.Sort([](const TSharedPtr<FObjectNode>& A, const TSharedPtr<FObjectNode>& B)
+		{
+			return FCString::Strcmp(A->GetObjectName(), B->GetObjectName()) < 0;
+		});
+
+	int32 ObjIndex = 0;
+	while (ObjIndex < ObjectNodes.Num())
+	{
+		if (InAsyncOperationProgress.ShouldCancelAsyncOp())
+		{
+			break;
+		}
+
+		TSharedPtr<FObjectNode>& ObjNode = ObjectNodes[ObjIndex];
+
+		int32 ObjEndIndex = ObjIndex + 1;
+		while (ObjEndIndex < ObjectNodes.Num() &&
+			FCString::Strcmp(ObjNode->GetObjectName(), ObjectNodes[ObjEndIndex]->GetObjectName()) == 0)
+		{
+			if (InAsyncOperationProgress.ShouldCancelAsyncOp())
+			{
+				break;
+			}
+			++ObjEndIndex;
+		}
+
+		if (ObjEndIndex - ObjIndex > 1)
+		{
+			const FText DisplayName = FText::FromString(ObjNode->GetObjectName());
+			auto ObjectGroupNode = MakeShared<FCustomTableTreeNode>(ObjNode->GetParentTable(), DisplayName);
+			ObjectGroupNode->SetColor(FLinearColor(0.9f, 0.9f, 1.0f, 1.0f));
+
+			for (int32 Index = ObjIndex; Index < ObjEndIndex; ++Index)
+			{
+				ObjectGroupNode->AddChildAndSetParent(ObjectNodes[Index]);
+			}
+			ParentGroup.AddChildAndSetParent(ObjectGroupNode);
+		}
+		else
+		{
+			ParentGroup.AddChildAndSetParent(ObjNode);
+		}
+
+		ObjIndex = ObjEndIndex;
+	}
+}
+
+////////////////////////////////////////////////////////////////////////////////////////////////////
+
+} // namespace UE::Insights::ObjectProfiler
+
+#undef LOCTEXT_NAMESPACE
diff --git a/Engine/Source/Developer/TraceInsights/Private/Insights/ObjectProfiler/ViewModels/ObjectGroupingByObjectName.h b/Engine/Source/Developer/TraceInsights/Private/Insights/ObjectProfiler/ViewModels/ObjectGroupingByObjectName.h
new file mode 100644
--- /dev/null
+++ b/Engine/Source/Developer/TraceInsights/Private/Insights/ObjectProfiler/ViewModels/ObjectGroupingByObjectName.h
@@ -0,0 +1,28 @@
+// Copyright Epic Games, Inc. All Rights Reserved.
+
+#pragma once
+
+#include "CoreTypes.h"
+
+// TraceInsightsCore
+#include "InsightsCore/Table/ViewModels/BaseTreeNode.h"
+#include "InsightsCore/Table/ViewModels/TableTreeNode.h"
+#include "InsightsCore/Table/ViewModels/TreeNodeGrouping.h"
+
+namespace UE::Insights::ObjectProfiler
+{
+
+class FObjectNode;
+
+class FObjectGroupingByObjectName : public FTreeNodeGrouping
+{
+	INSIGHTS_DECLARE_RTTI(FObjectGroupingByObjectName, FTreeNodeGrouping)
+
+public:
+	FObjectGroupingByObjectName();
+	virtual ~FObjectGroupingByObjectName() override;
+
+	virtual void GroupNodes(const TArray<FTableTreeNodePtr>& Nodes, FTableTreeNode& ParentGroup, TWeakPtr<FTable> InParentTable, IAsyncOperationProgress& InAsyncOperationProgress) const override;
+};
+
+} // namespace UE::Insights::ObjectProfiler
diff --git a/Engine/Source/Developer/TraceInsights/Private/Insights/ObjectProfiler/ViewModels/ObjectTableViewPresets.cpp b/Engine/Source/Developer/TraceInsights/Private/Insights/ObjectProfiler/ViewModels/ObjectTableViewPresets.cpp
--- a/Engine/Source/Developer/TraceInsights/Private/Insights/ObjectProfiler/ViewModels/ObjectTableViewPresets.cpp
+++ b/Engine/Source/Developer/TraceInsights/Private/Insights/ObjectProfiler/ViewModels/ObjectTableViewPresets.cpp
@@ -7,6 +7,7 @@
 #include "InsightsCore/Table/Widgets/STableTreeView.h"
 
 // TraceInsights
+#include "Insights/ObjectProfiler/ViewModels/ObjectGroupingByObjectName.h"
 #include "Insights/ObjectProfiler/ViewModels/ObjectGroupingByOuter.h"
 #include "Insights/ObjectProfiler/ViewModels/ObjectTable.h"
 #include "Insights/ObjectProfiler/Widgets/SObjectTableTreeView.h"
@@ -59,8 +60,7 @@ TSharedRef<ITableTreeViewPreset> FObjectTableViewPresets::CreateAssetViewPreset(
 			auto GroupingByObjectName = InAvailableGroupings.FindByPredicate(
 				[](TSharedPtr<FTreeNodeGrouping>& Grouping)
 				{
-					return Grouping->Is<FTreeNodeGroupingByPathBreakdown>() &&
-						Grouping->As<FTreeNodeGroupingByPathBreakdown>().GetColumnId() == FObjectTableColumns::ObjectNameColumnId;
+					return Grouping->Is<FObjectGroupingByObjectName>();
 				});
 			if (GroupingByObjectName)
 			{
diff --git a/Engine/Source/Developer/TraceInsights/Private/Insights/ObjectProfiler/Widgets/SObjectTableTreeView.cpp b/Engine/Source/Developer/TraceInsights/Private/Insights/ObjectProfiler/Widgets/SObjectTableTreeView.cpp
--- a/Engine/Source/Developer/TraceInsights/Private/Insights/ObjectProfiler/Widgets/SObjectTableTreeView.cpp
+++ b/Engine/Source/Developer/TraceInsights/Private/Insights/ObjectProfiler/Widgets/SObjectTableTreeView.cpp
@@ -28,6 +28,7 @@
 #include "Insights/InsightsStyle.h"
 #include "Insights/ObjectProfiler/IAssetInfoProvider.h"
 #include "Insights/ObjectProfiler/ObjectProfilerManager.h"
+#include "Insights/ObjectProfiler/ViewModels/ObjectGroupingByObjectName.h"
 #include "Insights/ObjectProfiler/ViewModels/ObjectGroupingByOuter.h"
 #include "Insights/ObjectProfiler/ViewModels/ObjectNode.h"
 #include "Insights/ObjectProfiler/ViewModels/ObjectTable.h"
@@ -655,9 +656,8 @@ void SObjectTableTreeView::InternalCreateGroupings()
 
 	int32 Index = 1; // after the Flat ("All") grouping
 
+	AvailableGroupings.Insert(MakeShared<FObjectGroupingByObjectName>(), Index++);
 	AvailableGroupings.Insert(MakeShared<FObjectGroupingByOuter>(), Index++);
-	//AvailableGroupings.Insert(MakeShared<FUObjectByClassNodeGrouping>(SharedThis(this)), Index++);
-	//AvailableGroupings.Insert(MakeShared<FUObjectByPackageNodeGrouping>(SharedThis(this)), Index++);
 }
 
 ////////////////////////////////////////////////////////////////////////////////////////////////////
