[Insights] ObjectProfiler: Added support for selecting multiple tree nodes with multiple sub-objects. It will show validations and actors for all detected sub-objects.

#Synced-CL 52443010

diff --git a/Engine/Source/Developer/TraceInsights/Private/Insights/ObjectProfiler/Widgets/SObjectDetailsView.cpp b/Engine/Source/Developer/TraceInsights/Private/Insights/ObjectProfiler/Widgets/SObjectDetailsView.cpp
--- a/Engine/Source/Developer/TraceInsights/Private/Insights/ObjectProfiler/Widgets/SObjectDetailsView.cpp
+++ b/Engine/Source/Developer/TraceInsights/Private/Insights/ObjectProfiler/Widgets/SObjectDetailsView.cpp
@@ -17,6 +17,9 @@
 #include "Widgets/Views/SListView.h"
 #include "Widgets/Views/STableRow.h"
 
+// TraceInsightsCore
+#include "InsightsCore/Table/ViewModels/BaseTreeNode.h"
+
 // TraceInsights
 #include "Insights/ObjectProfiler/IAssetInfoProvider.h"
 #include "Insights/ObjectProfiler/ViewModels/AssetInfoNode.h"
@@ -201,9 +204,9 @@
 		[
 			SNew(STextBlock)
 			.Font(FAppStyle::Get().GetFontStyle("Font.Large"))
-			.ColorAndOpacity(this, &SObjectDetailsView::GetColorForSelectedObjectName)
-			.Text(this, &SObjectDetailsView::GetSelectedObjectName)
-			.ToolTipText(this, &SObjectDetailsView::GetSelectedObjectToolTip)
+			.ColorAndOpacity(this, &SObjectDetailsView::GetTitleColor)
+			.Text(this, &SObjectDetailsView::GetTitleText)
+			.ToolTipText(this, &SObjectDetailsView::GetTitleToolTipText)
 		]
 
 		+ SVerticalBox::Slot()
@@ -212,8 +215,8 @@
 		.Padding(FMargin(8.0f, 0.0f, 8.0f, 4.0f))
 		[
 			SNew(STextBlock)
-			.Text(this, &SObjectDetailsView::GetSelectedObjectSizeText)
-			.ToolTipText(this, &SObjectDetailsView::GetSelectedObjectThumbnailToolTip)
+			.Text(this, &SObjectDetailsView::GetSubTitleText)
+			.ToolTipText(this, &SObjectDetailsView::GetSubTitleToolTipText)
 		]
 
 		+ SVerticalBox::Slot()
@@ -222,16 +225,16 @@
 		.Padding(FMargin(8.0f, 4.0f, 8.0f, 4.0f))
 		[
 			SNew(SBox)
-			.Visibility(this, &SObjectDetailsView::GetSelectedObjectVisibility)
+			.Visibility(this, &SObjectDetailsView::GetThumbnailVisibility)
 			.WidthOverride(50.0f)
 			.HeightOverride(50.0f)
-			.ToolTipText(this, &SObjectDetailsView::GetSelectedObjectThumbnailToolTip)
+			.ToolTipText(this, &SObjectDetailsView::GetThumbnailToolTipText)
 			[
 				SNew(SBorder)
 				.BorderBackgroundColor(FSlateColor::UseForeground())
 				[
 					SNew(SImage)
-					.Image(this, &SObjectDetailsView::GetSelectedObjectThumbnail)
+					.Image(this, &SObjectDetailsView::GetThumbnail)
 				]
 			]
 		]
@@ -306,61 +309,144 @@
 
 ////////////////////////////////////////////////////////////////////////////////////////////////////
 
-void SObjectDetailsView::SetObject(TSharedPtr<const FObjectNode> InObjectNode)
+FText SObjectDetailsView::GetSelectedObjectDetailedToolTip() const
 {
-	SelectedObjectNode = InObjectNode;
-	SelectedAssetData = FAssetData();
-	SelectedAssetClassName = NAME_None;
-	ValidationIssues.Empty();
+	if (SelectedObjectNode.IsValid())
+	{
+		TUtf8StringBuilder<1024> Str;
+		Str.Append(SelectedObjectNode->GetObjectName());
+		Str.Append("\nId: ");
+		Str.Appendf("%u", SelectedObjectNode->GetObjectId());
+		if (TSharedPtr<FObjectNode> ClassNode = SelectedObjectNode->GetClass())
+		{
+			Str.Append("\nClass: ");
+			Str.Append(ClassNode->GetObjectName());
+		}
+		if (TSharedPtr<FObjectNode> OuterNode = SelectedObjectNode->GetOuter())
+		{
+			Str.Append("\nOuter: ");
+			Str.Append(OuterNode->GetObjectName());
+		}
+		if (SelectedObjectNode->GetSystemMemorySize() > 0)
+		{
+			Str.Appendf("\nEstimated System Memory: %.2f MiB", double(SelectedObjectNode->GetSystemMemorySize()) / (1024.0 * 1024.0));
+		}
+		if (SelectedObjectNode->GetVideoMemorySize() > 0)
+		{
+			Str.Appendf("\nEstimated Video Memory: %.2f MiB", double(SelectedObjectNode->GetVideoMemorySize()) / (1024.0 * 1024.0));
+		}
+		return FText::FromString(Str.ToString());
+	}
+	return FText::GetEmpty();
+}
 
-	AssetEntries.Empty();
+////////////////////////////////////////////////////////////////////////////////////////////////////
 
-	if (InObjectNode.IsValid())
+FSlateColor SObjectDetailsView::GetTitleColor() const
+{
+	if (SelectedObjectNode.IsValid())
 	{
+		return FSlateColor::UseForeground();
+	}
+
+	if (SelectedNodes.Num() > 0)
+	{
+		return FSlateColor(EStyleColor::AccentFolder);
+	}
+
+	return FSlateColor(EStyleColor::White25);
+}
+
+////////////////////////////////////////////////////////////////////////////////////////////////////
+
+FText SObjectDetailsView::GetTitleText() const
+{
+	if (SelectedObjectNode.IsValid())
+	{
+		return FText::FromString(SelectedObjectNode->GetObjectName());
+	}
+
+	if (SelectedNodes.Num() > 0)
+	{
+		return FText::Format(LOCTEXT("MultipleItemsSelected", "{0} {0}|plural(one=Node,other=Nodes) Selected ({1} {1}|plural(one=Object,other=Objects))"),
+			FText::AsNumber(SelectedNodes.Num()),
+			FText::AsNumber(SelectedObjectNodes.Num()));
+	}
+
+	return LOCTEXT("NoObjectSelected", "No Object Selected");
+}
+
+////////////////////////////////////////////////////////////////////////////////////////////////////
+
+FText SObjectDetailsView::GetTitleToolTipText() const
+{
+	if (SelectedObjectNode.IsValid())
+	{
+		return FText::FromString(SelectedObjectNode->GetObjectPath());
+	}
+
+	return FText::GetEmpty();
+}
+
+////////////////////////////////////////////////////////////////////////////////////////////////////
+
+FText SObjectDetailsView::GetSubTitleText() const
+{
+	if (SelectedObjectNodes.Num() > 0)
+	{
+		uint64 TotalEstimatedMemorySize = 0;
+		for (const TSharedPtr<FObjectNode>& Node : SelectedObjectNodes)
+		{
+			TotalEstimatedMemorySize += Node->GetEstimatedMemorySize();
+		}
+		return FText::Format(LOCTEXT("SizeFmt", "Size: {0}"), FText::AsMemory(TotalEstimatedMemorySize));
+	}
+	return FText::GetEmpty();
+}
+
+////////////////////////////////////////////////////////////////////////////////////////////////////
+
+FText SObjectDetailsView::GetSubTitleToolTipText() const
+{
+	return GetSelectedObjectDetailedToolTip();
+}
+
+////////////////////////////////////////////////////////////////////////////////////////////////////
+
+EVisibility SObjectDetailsView::GetThumbnailVisibility() const
+{
+	return SelectedObjectNode.IsValid() ? EVisibility::Visible : EVisibility::Collapsed;
+}
+
+////////////////////////////////////////////////////////////////////////////////////////////////////
+
+FText SObjectDetailsView::GetThumbnailToolTipText() const
+{
+	return GetSelectedObjectDetailedToolTip();
+}
+
+////////////////////////////////////////////////////////////////////////////////////////////////////
+
+const FSlateBrush* SObjectDetailsView::GetThumbnail() const
+{
+	if (SelectedObjectNode.IsValid() &&
+		SelectedAssetData.IsValid() &&
+		SelectedAssetClassName != NAME_None)
+	{
 		auto AssetInfoProvider = WeakAssetInfoProvider.Pin();
 		if (AssetInfoProvider)
 		{
-			AssetInfoProvider->GetAssetData(*InObjectNode.Get(), SelectedAssetData);
-			TSharedPtr<FObjectNode> ClassNode = InObjectNode->GetClass();
-			SelectedAssetClassName = ClassNode ? FName(ClassNode->GetObjectName()) : NAME_None;
-
-			if (SelectedAssetData.IsValid())
-			{
-				AssetInfoProvider->GetValidationIssues(SelectedAssetData, ValidationIssues);
-			}
-
-			const TSet<FActorInfo>& Actors = InObjectNode->GetMatchedActors();
-			for (const FActorInfo& Actor : Actors)
-			{
-				auto Entry = MakeShared<FAssetInfoNode>();
-				Entry->Actor = Actor;
-				AssetEntries.Add(Entry);
-			}
+			return AssetInfoProvider->GetThumbnail(SelectedAssetData, SelectedAssetClassName);
 		}
 	}
-
-	ListView->RequestListRefresh();
+	return nullptr;
 }
 
 ////////////////////////////////////////////////////////////////////////////////////////////////////
 
-EVisibility SObjectDetailsView::GetSelectedObjectVisibility() const
-{
-	return SelectedObjectNode.IsValid() ? EVisibility::Visible : EVisibility::Collapsed;
-}
-
-////////////////////////////////////////////////////////////////////////////////////////////////////
-
-EVisibility SObjectDetailsView::GetActorsVisibility() const
-{
-	return SelectedObjectNode.IsValid() && AssetEntries.Num() > 0 ? EVisibility::Visible : EVisibility::Collapsed;
-}
-
-////////////////////////////////////////////////////////////////////////////////////////////////////
-
 EVisibility SObjectDetailsView::GetValidationIssuesVisibility() const
 {
-	return SelectedObjectNode.IsValid() && ValidationIssues.Num() > 0 ? EVisibility::Visible : EVisibility::Collapsed;
+	return ValidationIssues.Num() > 0 ? EVisibility::Visible : EVisibility::Collapsed;
 }
 
 ////////////////////////////////////////////////////////////////////////////////////////////////////
@@ -372,96 +458,94 @@
 
 ////////////////////////////////////////////////////////////////////////////////////////////////////
 
-FSlateColor SObjectDetailsView::GetColorForSelectedObjectName() const
+EVisibility SObjectDetailsView::GetActorsVisibility() const
 {
-	if (SelectedObjectNode.IsValid())
-	{
-		return FSlateColor::UseForeground();
-	}
-	return FSlateColor(EStyleColor::White25);
+	return AssetEntries.Num() > 0 ? EVisibility::Visible : EVisibility::Collapsed;
 }
 
 ////////////////////////////////////////////////////////////////////////////////////////////////////
 
-FText SObjectDetailsView::GetSelectedObjectName() const
+void SObjectDetailsView::SetSelectedNodes(const TArray<TSharedPtr<FTableTreeNode>>& InSelectedNodes)
 {
-	if (SelectedObjectNode.IsValid())
+	SelectedNodes = InSelectedNodes;
+	SelectedObjectNodes.Reset();
+
+	SelectedObjectNode = nullptr;
+	SelectedAssetData = FAssetData();
+	SelectedAssetClassName = NAME_None;
+	ValidationIssues.Reset();
+	UniqueActors.Reset();
+	AssetEntries.Reset();
+
+	for (const TSharedPtr<FTableTreeNode>& SelectedNode : SelectedNodes)
 	{
-		return FText::FromString(SelectedObjectNode->GetObjectName());
+		GatherObjectNodesRec(SelectedNode);
 	}
-	return LOCTEXT("NoObjectSelected", "No Object Selected");
+
+	if (SelectedObjectNodes.Num() == 1)
+	{
+		SelectedObjectNode = SelectedObjectNodes[0];
+
+		TSharedPtr<FObjectNode> ClassNode = SelectedObjectNode->GetClass();
+		SelectedAssetClassName = ClassNode ? FName(ClassNode->GetObjectName()) : NAME_None;
+
+	}
+	else if (SelectedObjectNodes.Num() > 1)
+	{
+		// Reset AssetData if multiple objects are selected.
+		SelectedAssetData = FAssetData();
+	}
+
+	ListView->RequestListRefresh();
 }
 
 ////////////////////////////////////////////////////////////////////////////////////////////////////
 
-FText SObjectDetailsView::GetSelectedObjectSizeText() const
+void SObjectDetailsView::GatherObjectNodesRec(TSharedPtr<FBaseTreeNode> Node)
 {
-	if (SelectedObjectNode.IsValid())
+	if (Node->Is<FObjectNode>())
 	{
-		return FText::Format(LOCTEXT("SizeFmt", "Size: {0}"), FText::AsMemory(SelectedObjectNode->GetEstimatedMemorySize()));
+		TSharedPtr<FObjectNode> ObjectNode = StaticCastSharedPtr<FObjectNode>(Node);
+		AddObjectNode(ObjectNode);
 	}
-	return FText::GetEmpty();
+
+	for (FBaseTreeNodePtr ChildPtr : Node->GetChildren())
+	{
+		GatherObjectNodesRec(ChildPtr);
+	}
 }
 
 ////////////////////////////////////////////////////////////////////////////////////////////////////
 
-FText SObjectDetailsView::GetSelectedObjectToolTip() const
+void SObjectDetailsView::AddObjectNode(TSharedPtr<FObjectNode> InObjectNode)
 {
-	if (SelectedObjectNode.IsValid())
+	SelectedObjectNodes.Add(InObjectNode);
+
+	auto AssetInfoProvider = WeakAssetInfoProvider.Pin();
+	if (!AssetInfoProvider)
 	{
-		return FText::FromString(SelectedObjectNode->GetObjectPath());
+		return;
 	}
-	return FText::GetEmpty();
+
+	AssetInfoProvider->GetAssetData(*InObjectNode.Get(), SelectedAssetData);
+
+	if (SelectedAssetData.IsValid())
+	{
+		AssetInfoProvider->GetValidationIssues(SelectedAssetData, ValidationIssues);
+	}
+
+	const TSet<FActorInfo>& Actors = InObjectNode->GetMatchedActors();
+	for (const FActorInfo& Actor : Actors)
+	{
+		if (!UniqueActors.Contains(Actor))
+		{
+			UniqueActors.Add(Actor);
+
+			auto Entry = MakeShared<FAssetInfoNode>();
+			Entry->Actor = Actor;
+			AssetEntries.Add(Entry);
+		}
+	}
 }
 
 ////////////////////////////////////////////////////////////////////////////////////////////////////
-
-FText SObjectDetailsView::GetSelectedObjectThumbnailToolTip() const
-{
-	if (SelectedObjectNode.IsValid())
-	{
-		TUtf8StringBuilder<1024> Str;
-		Str.Append(SelectedObjectNode->GetObjectName());
-		Str.Append("\nId: ");
-		Str.Appendf("%u", SelectedObjectNode->GetObjectId());
-		if (TSharedPtr<FObjectNode> ClassNode = SelectedObjectNode->GetClass())
-		{
-			Str.Append("\nClass: ");
-			Str.Append(ClassNode->GetObjectName());
-		}
-		if (TSharedPtr<FObjectNode> OuterNode = SelectedObjectNode->GetOuter())
-		{
-			Str.Append("\nOuter: ");
-			Str.Append(OuterNode->GetObjectName());
-		}
-		if (SelectedObjectNode->GetSystemMemorySize() > 0)
-		{
-			Str.Appendf("\nEstimated System Memory: %.2f MiB", double(SelectedObjectNode->GetSystemMemorySize()) / (1024.0 * 1024.0));
-		}
-		if (SelectedObjectNode->GetVideoMemorySize() > 0)
-		{
-			Str.Appendf("\nEstimated Video Memory: %.2f MiB", double(SelectedObjectNode->GetVideoMemorySize()) / (1024.0 * 1024.0));
-		}
-		return FText::FromString(Str.ToString());
-	}
-	return FText::GetEmpty();
-}
-
-////////////////////////////////////////////////////////////////////////////////////////////////////
-
-const FSlateBrush* SObjectDetailsView::GetSelectedObjectThumbnail() const
-{
-	if (SelectedObjectNode.IsValid() &&
-		SelectedAssetData.IsValid() &&
-		SelectedAssetClassName != NAME_None)
-	{
-		auto AssetInfoProvider = WeakAssetInfoProvider.Pin();
-		if (AssetInfoProvider)
-		{
-			return AssetInfoProvider->GetThumbnail(SelectedAssetData, SelectedAssetClassName);
-		}
-	}
-	return nullptr;
-}
-
-////////////////////////////////////////////////////////////////////////////////////////////////////
diff --git a/Engine/Source/Developer/TraceInsights/Private/Insights/ObjectProfiler/Widgets/SObjectDetailsView.h b/Engine/Source/Developer/TraceInsights/Private/Insights/ObjectProfiler/Widgets/SObjectDetailsView.h
--- a/Engine/Source/Developer/TraceInsights/Private/Insights/ObjectProfiler/Widgets/SObjectDetailsView.h
+++ b/Engine/Source/Developer/TraceInsights/Private/Insights/ObjectProfiler/Widgets/SObjectDetailsView.h
@@ -11,6 +11,11 @@
 // TraceInsights
 #include "Insights/ObjectProfiler/ViewModels/AssetInfoNode.h"
 
+namespace UE::Insights
+{
+	class FBaseTreeNode;
+}
+
 namespace UE::Insights::ObjectProfiler
 {
 
@@ -56,21 +61,31 @@ public:
 		WeakAssetInfoProvider = InAssetInfoProvider.ToWeakPtr();
 	}
 
-	void SetObject(TSharedPtr<const FObjectNode> InObjectNode);
+	void SetSelectedNodes(const TArray<TSharedPtr<FTableTreeNode>>& InSelectedNodes);
 
-protected:
+private:
 	TSharedRef<ITableRow> OnGenerateRow(TSharedPtr<FAssetInfoNode> InItem, const TSharedRef<STableViewBase>& OwnerTable);
 
-	EVisibility GetSelectedObjectVisibility() const;
+	FText GetSelectedObjectDetailedToolTip() const;
+
+	FSlateColor GetTitleColor() const;
+	FText GetTitleText() const;
+	FText GetTitleToolTipText() const;
+
+	FText GetSubTitleText() const;
+	FText GetSubTitleToolTipText() const;
+
+	EVisibility GetThumbnailVisibility() const;
+	FText GetThumbnailToolTipText() const;
+	const FSlateBrush* GetThumbnail() const;
+
 	EVisibility GetValidationIssuesVisibility() const;
-	EVisibility GetActorsVisibility() const;
-	FSlateColor GetColorForSelectedObjectName() const;
-	FText GetSelectedObjectName() const;
 	FText GetValidationIssuesText() const;
-	FText GetSelectedObjectSizeText() const;
-	FText GetSelectedObjectToolTip() const;
-	FText GetSelectedObjectThumbnailToolTip() const;
-	const FSlateBrush* GetSelectedObjectThumbnail() const;
+
+	EVisibility GetActorsVisibility() const;
+
+	void GatherObjectNodesRec(TSharedPtr<FBaseTreeNode> Node);
+	void AddObjectNode(TSharedPtr<FObjectNode> InObjectNode);
 
 private:
 	TWeakPtr<IAssetInfoProvider> WeakAssetInfoProvider;
@@ -78,11 +93,17 @@ private:
 	TWeakPtr<SObjectProfilerWindow> ParentWindow;
 	TWeakPtr<SMemoryView> ParentMemoryView;
 
-	TSharedPtr<const FObjectNode> SelectedObjectNode;
+	TArray<TSharedPtr<FTableTreeNode>> SelectedNodes;
+	TArray<TSharedPtr<FObjectNode>> SelectedObjectNodes;
+
+	TSharedPtr<FObjectNode> SelectedObjectNode;
 	FAssetData SelectedAssetData;
 	FName SelectedAssetClassName;
+
 	TArray<FText> ValidationIssues;
 
+	TSet<FActorInfo> UniqueActors;
+
 	/** List of matched asset info nodes shown in the list view. */
 	TArray<TSharedPtr<FAssetInfoNode>> AssetEntries;
 
diff --git a/Engine/Source/Developer/TraceInsights/Private/Insights/ObjectProfiler/Widgets/SObjectTableTreeView.cpp b/Engine/Source/Developer/TraceInsights/Private/Insights/ObjectProfiler/Widgets/SObjectTableTreeView.cpp
--- a/Engine/Source/Developer/TraceInsights/Private/Insights/ObjectProfiler/Widgets/SObjectTableTreeView.cpp
+++ b/Engine/Source/Developer/TraceInsights/Private/Insights/ObjectProfiler/Widgets/SObjectTableTreeView.cpp
@@ -309,8 +309,6 @@ void SObjectTableTreeView::UpdateSelectionStatsText()
 	TArray<FTableTreeNodePtr> SelectedNodes;
 	const int32 NumSelectedNodes = TreeView->GetSelectedItems(SelectedNodes);
 
-	TSharedPtr<const FObjectNode> FirstSelectedNode;
-
 	if (NumSelectedNodes > 0)
 	{
 		int64 TotalCount = 0;
@@ -321,11 +319,6 @@ void SObjectTableTreeView::UpdateSelectionStatsText()
 		{
 			if (SelectedNode->Is<FObjectNode>())
 			{
-				if (!FirstSelectedNode.IsValid())
-				{
-					FirstSelectedNode = StaticCastSharedPtr<const FObjectNode>(SelectedNode);
-				}
-
 				const FObjectNode& ObjectNode = SelectedNode->As<FObjectNode>();
 				++TotalCount;
 				TotalStructSize += ObjectNode.GetStructureSize();
@@ -354,7 +347,7 @@ void SObjectTableTreeView::UpdateSelectionStatsText()
 
 	if (TSharedPtr<SObjectDetailsView> ObjectDetailsView = WeakObjectDetailsView.Pin())
 	{
-		ObjectDetailsView->SetObject(FirstSelectedNode);
+		ObjectDetailsView->SetSelectedNodes(SelectedNodes);
 	}
 }
 
