[Insights] Removed unused code (ProfilerMenuBuilder::AddMenuEntry) from profiler commands. Simplified code re commands (added reusable Common/CommandUtils.h with INSIGHTS_DECLARE_TOGGLE_COMMAND and INSIGHTS_IMPLEMENT_TOGGLE_COMMAND).

#rnx
#Synced-CL 52471416

diff --git a/Engine/Source/Developer/TraceInsights/Private/Insights/Common/CommandUtils.h b/Engine/Source/Developer/TraceInsights/Private/Insights/Common/CommandUtils.h
new file mode 100644
--- /dev/null
+++ b/Engine/Source/Developer/TraceInsights/Private/Insights/Common/CommandUtils.h
@@ -0,0 +1,40 @@
+// Copyright Epic Games, Inc. All Rights Reserved.
+
+#pragma once
+
+#include "CoreTypes.h"
+
+#define INSIGHTS_DECLARE_TOGGLE_COMMAND(CmdName) \
+	public: \
+		void Map_##CmdName(); /**< Maps UI command info CmdName with the specified UI command list. */ \
+		const FUIAction GetAction_##CmdName(); /**< UI action for CmdName command. */ \
+	protected: \
+		void CmdName##_Execute(); /**< Handles FExecuteAction for CmdName. */ \
+		bool CmdName##_CanExecute() const; /**< Handles FCanExecuteAction for CmdName. */ \
+		ECheckBoxState CmdName##_GetCheckState() const; /**< Handles FGetActionCheckState for CmdName. */
+
+#define INSIGHTS_IMPLEMENT_TOGGLE_COMMAND(Class, This, CmdName, IsEnabled, SetIsEnabled) \
+	void Class::Map_##CmdName() \
+	{ \
+		This->GetCommandList()->MapAction(This->GetCommands().CmdName, GetAction_##CmdName()); \
+	} \
+	const FUIAction Class::GetAction_##CmdName() \
+	{ \
+		FUIAction UIAction; \
+		UIAction.ExecuteAction = FExecuteAction::CreateRaw(this, &Class::CmdName##_Execute); \
+		UIAction.CanExecuteAction = FCanExecuteAction::CreateRaw(this, &Class::CmdName##_CanExecute); \
+		UIAction.GetActionCheckState = FGetActionCheckState::CreateRaw(this, &Class::CmdName##_GetCheckState); \
+		return UIAction; \
+	} \
+	void Class::CmdName##_Execute() \
+	{ \
+		This->SetIsEnabled(!This->IsEnabled()); \
+	} \
+	bool Class::CmdName##_CanExecute() const \
+	{ \
+		return FInsightsManager::Get()->GetSession().IsValid(); \
+	} \
+	ECheckBoxState Class::CmdName##_GetCheckState() const \
+	{ \
+		return This->IsEnabled() ? ECheckBoxState::Checked : ECheckBoxState::Unchecked; \
+	}
diff --git a/Engine/Source/Developer/TraceInsights/Private/Insights/InsightsCommands.cpp b/Engine/Source/Developer/TraceInsights/Private/Insights/InsightsCommands.cpp
--- a/Engine/Source/Developer/TraceInsights/Private/Insights/InsightsCommands.cpp
+++ b/Engine/Source/Developer/TraceInsights/Private/Insights/InsightsCommands.cpp
@@ -37,7 +37,7 @@ FInsightsCommands::FInsightsCommands()
 UE_DISABLE_OPTIMIZATION_SHIP
 void FInsightsCommands::RegisterCommands()
 {
-	UI_COMMAND(InsightsManager_Load,
+	UI_COMMAND(LoadTraceFile,
 		"Load...",
 		"Loads profiler data from a trace file.",
 		EUserInterfaceActionType::Button,
@@ -58,56 +58,27 @@ void FInsightsCommands::RegisterCommands()
 UE_ENABLE_OPTIMIZATION_SHIP
 
 ////////////////////////////////////////////////////////////////////////////////////////////////////
-// ToggleDebugInfo
+// Toggle Commands
 ////////////////////////////////////////////////////////////////////////////////////////////////////
 
-void FInsightsActionManager::Map_ToggleDebugInfo_Global()
-{
-	This->CommandList->MapAction(This->GetCommands().ToggleDebugInfo, ToggleDebugInfo_Custom());
-}
-
-const FUIAction FInsightsActionManager::ToggleDebugInfo_Custom()
-{
-	FUIAction UIAction;
-	UIAction.ExecuteAction = FExecuteAction::CreateRaw(this, &FInsightsActionManager::ToggleDebugInfo_Execute);
-	UIAction.CanExecuteAction = FCanExecuteAction::CreateRaw(this, &FInsightsActionManager::ToggleDebugInfo_CanExecute);
-	UIAction.GetActionCheckState = FGetActionCheckState::CreateRaw(this, &FInsightsActionManager::ToggleDebugInfo_GetCheckState);
-	return UIAction;
-}
-
-void FInsightsActionManager::ToggleDebugInfo_Execute()
-{
-	const bool bIsDebugInfoEnabled = !This->IsDebugInfoEnabled();
-	This->SetDebugInfo(bIsDebugInfoEnabled);
-}
-
-bool FInsightsActionManager::ToggleDebugInfo_CanExecute() const
-{
-	return FInsightsManager::Get()->GetSession().IsValid();
-}
-
-ECheckBoxState FInsightsActionManager::ToggleDebugInfo_GetCheckState() const
-{
-	const bool bIsDebugInfoEnabled = This->IsDebugInfoEnabled();
-	return bIsDebugInfoEnabled ? ECheckBoxState::Checked : ECheckBoxState::Unchecked;
-}
+INSIGHTS_IMPLEMENT_TOGGLE_COMMAND(FInsightsActionManager, This, ToggleDebugInfo, IsDebugInfoEnabled, SetDebugInfo)
 
 ////////////////////////////////////////////////////////////////////////////////////////////////////
-// InsightsManager_Load
+// LoadTraceFile
 ////////////////////////////////////////////////////////////////////////////////////////////////////
 
-void FInsightsActionManager::Map_InsightsManager_Load()
+void FInsightsActionManager::Map_LoadTraceFile()
 {
 	FUIAction UIAction;
-	UIAction.ExecuteAction = FExecuteAction::CreateRaw(this, &FInsightsActionManager::InsightsManager_Load_Execute);
-	UIAction.CanExecuteAction = FCanExecuteAction::CreateRaw(this, &FInsightsActionManager::InsightsManager_Load_CanExecute);
+	UIAction.ExecuteAction = FExecuteAction::CreateRaw(this, &FInsightsActionManager::LoadTraceFile_Execute);
+	UIAction.CanExecuteAction = FCanExecuteAction::CreateRaw(this, &FInsightsActionManager::LoadTraceFile_CanExecute);
 
-	This->CommandList->MapAction(This->GetCommands().InsightsManager_Load, UIAction);
+	This->CommandList->MapAction(This->GetCommands().LoadTraceFile, UIAction);
 }
 
 ////////////////////////////////////////////////////////////////////////////////////////////////////
 
-void FInsightsActionManager::InsightsManager_Load_Execute()
+void FInsightsActionManager::LoadTraceFile_Execute()
 {
 	//const FString ProfilingDirectory(FPaths::ConvertRelativePathToFull(*FPaths::ProfilingDir()));
 	const FString ProfilingDirectory(This->GetStoreDir());
@@ -144,7 +115,7 @@ void FInsightsActionManager::InsightsManager_Load_Execute()
 
 ////////////////////////////////////////////////////////////////////////////////////////////////////
 
-bool FInsightsActionManager::InsightsManager_Load_CanExecute() const
+bool FInsightsActionManager::LoadTraceFile_CanExecute() const
 {
 	return true;
 }
@@ -153,19 +124,13 @@ bool FInsightsActionManager::InsightsManager_Load_CanExecute() const
 // OpenSettings
 ////////////////////////////////////////////////////////////////////////////////////////////////////
 
-void FInsightsActionManager::Map_OpenSettings_Global()
-{
-	This->CommandList->MapAction(This->GetCommands().OpenSettings, OpenSettings_Custom());
-}
-
-////////////////////////////////////////////////////////////////////////////////////////////////////
-
-const FUIAction FInsightsActionManager::OpenSettings_Custom()
+void FInsightsActionManager::Map_OpenSettings()
 {
 	FUIAction UIAction;
 	UIAction.ExecuteAction = FExecuteAction::CreateRaw(this, &FInsightsActionManager::OpenSettings_Execute);
 	UIAction.CanExecuteAction = FCanExecuteAction::CreateRaw(this, &FInsightsActionManager::OpenSettings_CanExecute);
-	return UIAction;
+
+	This->CommandList->MapAction(This->GetCommands().OpenSettings, UIAction);
 }
 
 ////////////////////////////////////////////////////////////////////////////////////////////////////
diff --git a/Engine/Source/Developer/TraceInsights/Private/Insights/InsightsCommands.h b/Engine/Source/Developer/TraceInsights/Private/Insights/InsightsCommands.h
--- a/Engine/Source/Developer/TraceInsights/Private/Insights/InsightsCommands.h
+++ b/Engine/Source/Developer/TraceInsights/Private/Insights/InsightsCommands.h
@@ -7,6 +7,9 @@
 #include "Framework/Commands/Commands.h"
 #include "Framework/Commands/UIAction.h"
 
+// TraceInsights
+#include "Insights/Common/CommandUtils.h"
+
 namespace UE::Insights
 {
 
@@ -26,19 +29,12 @@ public:
 	virtual void RegisterCommands() override;
 
 public:
-	//////////////////////////////////////////////////
-	// Global commands need to implement following method:
-	//     void Map_<CommandName>_Global();
-	// Custom commands needs to implement also the following method:
-	//     const FUIAction <CommandName>_Custom(...) const;
-	//////////////////////////////////////////////////
-
-	/** Load profiler data from a trace file. Global version. */
-	TSharedPtr<FUICommandInfo> InsightsManager_Load;
-
-	/** Toggles the debug info. Global and custom command. */
+	/** Toggles the debug info. */
 	TSharedPtr<FUICommandInfo> ToggleDebugInfo;
 
+	/** Load profiler data from a trace file. */
+	TSharedPtr<FUICommandInfo> LoadTraceFile;
+
 	/** Open settings for the profiler manager. */
 	TSharedPtr<FUICommandInfo> OpenSettings;
 };
@@ -59,34 +55,27 @@ private:
 	{}
 
 	//////////////////////////////////////////////////
-	// InsightsManager_Load
+	// Toggle Commands
 
-public:
-	void Map_InsightsManager_Load(); /**< Maps UI command info InsightsManager_Load with the specified UI command list. */
-protected:
-	void InsightsManager_Load_Execute(); /**< Handles FExecuteAction for InsightsManager_Load. */
-	bool InsightsManager_Load_CanExecute() const; /**< Handles FCanExecuteAction for InsightsManager_Load. */
+	INSIGHTS_DECLARE_TOGGLE_COMMAND(ToggleDebugInfo)
 
 	//////////////////////////////////////////////////
-	// ToggleDebugInfo
+	// LoadTraceFile Command
 
 public:
-	void Map_ToggleDebugInfo_Global(); /**< Maps UI command info ToggleDebugInfo with the specified UI command list. */
-	const FUIAction ToggleDebugInfo_Custom(); /**< UI action for ToggleDebugInfo command. */
+	void Map_LoadTraceFile(); /**< Maps UI command info LoadTraceFile with the specified UI command list. */
 protected:
-	void ToggleDebugInfo_Execute(); /**< Handles FExecuteAction for ToggleDebugInfo. */
-	bool ToggleDebugInfo_CanExecute() const; /**< Handles FCanExecuteAction for ToggleDebugInfo. */
-	ECheckBoxState ToggleDebugInfo_GetCheckState() const; /**< Handles FGetActionCheckState for ToggleDebugInfo. */
+	void LoadTraceFile_Execute();
+	bool LoadTraceFile_CanExecute() const;
 
 	//////////////////////////////////////////////////
-	// OpenSettings
+	// OpenSettings Command
 
 public:
-	void Map_OpenSettings_Global(); /**< Maps UI command info OpenSettings with the specified UI command list. */
-	const FUIAction OpenSettings_Custom(); /**< UI action for OpenSettings command. */
+	void Map_OpenSettings(); /**< Maps UI command info OpenSettings with the specified UI command list. */
 protected:
-	void OpenSettings_Execute(); /**< Handles FExecuteAction for OpenSettings. */
-	bool OpenSettings_CanExecute() const; /**< Handles FCanExecuteAction for OpenSettings. */
+	void OpenSettings_Execute();
+	bool OpenSettings_CanExecute() const;
 
 	//////////////////////////////////////////////////
 
diff --git a/Engine/Source/Developer/TraceInsights/Private/Insights/InsightsManager.cpp b/Engine/Source/Developer/TraceInsights/Private/Insights/InsightsManager.cpp
--- a/Engine/Source/Developer/TraceInsights/Private/Insights/InsightsManager.cpp
+++ b/Engine/Source/Developer/TraceInsights/Private/Insights/InsightsManager.cpp
@@ -175,9 +175,9 @@ FInsightsManager::~FInsightsManager()
 
 void FInsightsManager::BindCommands()
 {
-	ActionManager.Map_InsightsManager_Load();
-	ActionManager.Map_ToggleDebugInfo_Global();
-	ActionManager.Map_OpenSettings_Global();
+	ActionManager.Map_LoadTraceFile();
+	ActionManager.Map_ToggleDebugInfo();
+	ActionManager.Map_OpenSettings();
 }
 
 ////////////////////////////////////////////////////////////////////////////////////////////////////
diff --git a/Engine/Source/Developer/TraceInsights/Private/Insights/LoadingProfiler/LoadingProfilerCommands.cpp b/Engine/Source/Developer/TraceInsights/Private/Insights/LoadingProfiler/LoadingProfilerCommands.cpp
--- a/Engine/Source/Developer/TraceInsights/Private/Insights/LoadingProfiler/LoadingProfilerCommands.cpp
+++ b/Engine/Source/Developer/TraceInsights/Private/Insights/LoadingProfiler/LoadingProfilerCommands.cpp
@@ -4,7 +4,6 @@
 
 #include "DesktopPlatformModule.h"
 #include "Framework/Application/SlateApplication.h"
-#include "Framework/MultiBox/MultiBoxBuilder.h"
 
 // TraceInsights
 #include "Insights/InsightsManager.h"
@@ -16,23 +15,6 @@
 namespace UE::Insights::LoadingProfiler
 {
 
-////////////////////////////////////////////////////////////////////////////////////////////////////
-// FLoadingProfilerMenuBuilder
-////////////////////////////////////////////////////////////////////////////////////////////////////
-
-void FLoadingProfilerMenuBuilder::AddMenuEntry(FMenuBuilder& MenuBuilder, const TSharedPtr< FUICommandInfo >& UICommandInfo, const FUIAction& UIAction)
-{
-	MenuBuilder.AddMenuEntry
-	(
-		UICommandInfo->GetLabel(),
-		UICommandInfo->GetDescription(),
-		UICommandInfo->GetIcon(),
-		UIAction,
-		NAME_None,
-		UICommandInfo->GetUserInterfaceType()
-	);
-}
-
 ////////////////////////////////////////////////////////////////////////////////////////////////////
 // FLoadingProfilerCommands
 ////////////////////////////////////////////////////////////////////////////////////////////////////
@@ -94,47 +76,12 @@ UE_ENABLE_OPTIMIZATION_SHIP
 // Toggle Commands
 ////////////////////////////////////////////////////////////////////////////////////////////////////
 
-#define IMPLEMENT_TOGGLE_COMMAND(CmdName, IsEnabled, SetIsEnabled) \
-	\
-	void FLoadingProfilerActionManager::Map_##CmdName##_Global()\
-	{\
-		This->CommandList->MapAction(This->GetCommands().CmdName, CmdName##_Custom());\
-	}\
-	\
-	const FUIAction FLoadingProfilerActionManager::CmdName##_Custom() \
-	{\
-		FUIAction UIAction;\
-		UIAction.ExecuteAction = FExecuteAction::CreateRaw(this, &FLoadingProfilerActionManager::CmdName##_Execute);\
-		UIAction.CanExecuteAction = FCanExecuteAction::CreateRaw(this, &FLoadingProfilerActionManager::CmdName##_CanExecute);\
-		UIAction.GetActionCheckState = FGetActionCheckState::CreateRaw(this, &FLoadingProfilerActionManager::CmdName##_GetCheckState);\
-		return UIAction;\
-	}\
-	\
-	void FLoadingProfilerActionManager::CmdName##_Execute()\
-	{\
-		const bool b##IsEnabled = !This->IsEnabled();\
-		This->SetIsEnabled(b##IsEnabled);\
-	}\
-	\
-	bool FLoadingProfilerActionManager::CmdName##_CanExecute() const\
-	{\
-		return FInsightsManager::Get()->GetSession().IsValid();\
-	}\
-	\
-	ECheckBoxState FLoadingProfilerActionManager::CmdName##_GetCheckState() const\
-	{\
-		const bool b##IsEnabled = This->IsEnabled();\
-		return b##IsEnabled ? ECheckBoxState::Checked : ECheckBoxState::Unchecked;\
-	}
-
-IMPLEMENT_TOGGLE_COMMAND(ToggleTimingViewVisibility, IsTimingViewVisible, ShowHideTimingView)
-IMPLEMENT_TOGGLE_COMMAND(ToggleEventAggregationTreeViewVisibility, IsEventAggregationTreeViewVisible, ShowHideEventAggregationTreeView)
-IMPLEMENT_TOGGLE_COMMAND(ToggleObjectTypeAggregationTreeViewVisibility, IsObjectTypeAggregationTreeViewVisible, ShowHideObjectTypeAggregationTreeView)
-IMPLEMENT_TOGGLE_COMMAND(TogglePackageDetailsTreeViewVisibility, IsPackageDetailsTreeViewVisible, ShowHidePackageDetailsTreeView)
-IMPLEMENT_TOGGLE_COMMAND(ToggleExportDetailsTreeViewVisibility, IsExportDetailsTreeViewVisible, ShowHideExportDetailsTreeView)
-IMPLEMENT_TOGGLE_COMMAND(ToggleRequestsTreeViewVisibility, IsRequestsTreeViewVisible, ShowHideRequestsTreeView)
-
-#undef IMPLEMENT_TOGGLE_COMMAND
+INSIGHTS_IMPLEMENT_TOGGLE_COMMAND(FLoadingProfilerActionManager, This, ToggleTimingViewVisibility, IsTimingViewVisible, ShowHideTimingView)
+INSIGHTS_IMPLEMENT_TOGGLE_COMMAND(FLoadingProfilerActionManager, This, ToggleEventAggregationTreeViewVisibility, IsEventAggregationTreeViewVisible, ShowHideEventAggregationTreeView)
+INSIGHTS_IMPLEMENT_TOGGLE_COMMAND(FLoadingProfilerActionManager, This, ToggleObjectTypeAggregationTreeViewVisibility, IsObjectTypeAggregationTreeViewVisible, ShowHideObjectTypeAggregationTreeView)
+INSIGHTS_IMPLEMENT_TOGGLE_COMMAND(FLoadingProfilerActionManager, This, TogglePackageDetailsTreeViewVisibility, IsPackageDetailsTreeViewVisible, ShowHidePackageDetailsTreeView)
+INSIGHTS_IMPLEMENT_TOGGLE_COMMAND(FLoadingProfilerActionManager, This, ToggleExportDetailsTreeViewVisibility, IsExportDetailsTreeViewVisible, ShowHideExportDetailsTreeView)
+INSIGHTS_IMPLEMENT_TOGGLE_COMMAND(FLoadingProfilerActionManager, This, ToggleRequestsTreeViewVisibility, IsRequestsTreeViewVisible, ShowHideRequestsTreeView)
 
 ////////////////////////////////////////////////////////////////////////////////////////////////////
 
diff --git a/Engine/Source/Developer/TraceInsights/Private/Insights/LoadingProfiler/LoadingProfilerCommands.h b/Engine/Source/Developer/TraceInsights/Private/Insights/LoadingProfiler/LoadingProfilerCommands.h
--- a/Engine/Source/Developer/TraceInsights/Private/Insights/LoadingProfiler/LoadingProfilerCommands.h
+++ b/Engine/Source/Developer/TraceInsights/Private/Insights/LoadingProfiler/LoadingProfilerCommands.h
@@ -9,7 +9,8 @@
 #include "Styling/SlateTypes.h" // for ECheckBoxState
 #include "Templates/SharedPointer.h"
 
-class FMenuBuilder;
+// TraceInsights
+#include "Insights/Common/CommandUtils.h"
 
 namespace UE::Insights::LoadingProfiler
 {
@@ -28,54 +29,25 @@ public:
 	virtual void RegisterCommands() override;
 
 public:
-	//////////////////////////////////////////////////
-	// Global commands need to implement following method:
-	//     void Map_<CommandName>_Global();
-	// Custom commands needs to implement also the following method:
-	//     const FUIAction <CommandName>_Custom(...) const;
-	//////////////////////////////////////////////////
-
-	/** Toggles visibility for the Timing view. Global and custom command. */
+	/** Toggles visibility for the Timing view. */
 	TSharedPtr<FUICommandInfo> ToggleTimingViewVisibility;
 
-	/** Toggles visibility for the Event Aggregation tree view. Global and custom command. */
+	/** Toggles visibility for the Event Aggregation tree view. */
 	TSharedPtr<FUICommandInfo> ToggleEventAggregationTreeViewVisibility;
 
-	/** Toggles visibility for the Object Type Aggregation tree view. Global and custom command. */
+	/** Toggles visibility for the Object Type Aggregation tree view. */
 	TSharedPtr<FUICommandInfo> ToggleObjectTypeAggregationTreeViewVisibility;
 
-	/** Toggles visibility for the Package Details tree view. Global and custom command. */
+	/** Toggles visibility for the Package Details tree view. */
 	TSharedPtr<FUICommandInfo> TogglePackageDetailsTreeViewVisibility;
 
-	/** Toggles visibility for the Export Details tree view. Global and custom command. */
+	/** Toggles visibility for the Export Details tree view. */
 	TSharedPtr<FUICommandInfo> ToggleExportDetailsTreeViewVisibility;
 
-	/** Toggles visibility for the Requests tree view. Global and custom command. */
+	/** Toggles visibility for the Requests tree view. */
 	TSharedPtr<FUICommandInfo> ToggleRequestsTreeViewVisibility;
 };
 
-////////////////////////////////////////////////////////////////////////////////////////////////////
-/**
- * Menu builder. Helper class for adding a customized menu entry using the global UI command info.
- */
-class FLoadingProfilerMenuBuilder
-{
-public:
-	/**
-	 * Helper method for adding a customized menu entry using the global UI command info.
-	 * FUICommandInfo cannot be executed with custom parameters, so we need to create a custom FUIAction,
-	 * but sometime we have global and local version for the UI command, so reuse data from the global UI command info.
-	 * Ex:
-	 *     SessionInstance_ToggleCapture          - Global version will toggle capture process for all active session instances
-	 *     SessionInstance_ToggleCapture_OneParam - Local version will toggle capture process only for the specified session instance
-	 *
-	 * @param MenuBuilder The menu to add items to
-	 * @param FUICommandInfo A shared pointer to the UI command info
-	 * @param UIAction Customized version of the UI command info stored in an UI action
-	 */
-	static void AddMenuEntry(FMenuBuilder& MenuBuilder, const TSharedPtr<FUICommandInfo>& UICommandInfo, const FUIAction& UIAction);
-};
-
 ////////////////////////////////////////////////////////////////////////////////////////////////////
 /**
  * Class that provides helper functions for the commands to avoid cluttering profiler manager with many small functions.
@@ -94,22 +66,12 @@ private:
 	//////////////////////////////////////////////////
 	// Toggle Commands
 
-#define DECLARE_TOGGLE_COMMAND(CmdName)\
-public:\
-	void Map_##CmdName##_Global(); /**< Maps UI command info CmdName with the specified UI command list. */\
-	const FUIAction CmdName##_Custom(); /**< UI action for CmdName command. */\
-protected:\
-	void CmdName##_Execute(); /**< Handles FExecuteAction for CmdName. */\
-	bool CmdName##_CanExecute() const; /**< Handles FCanExecuteAction for CmdName. */\
-	ECheckBoxState CmdName##_GetCheckState() const; /**< Handles FGetActionCheckState for CmdName. */
-
-	DECLARE_TOGGLE_COMMAND(ToggleTimingViewVisibility)
-	DECLARE_TOGGLE_COMMAND(ToggleEventAggregationTreeViewVisibility)
-	DECLARE_TOGGLE_COMMAND(ToggleObjectTypeAggregationTreeViewVisibility)
-	DECLARE_TOGGLE_COMMAND(TogglePackageDetailsTreeViewVisibility)
-	DECLARE_TOGGLE_COMMAND(ToggleExportDetailsTreeViewVisibility)
-	DECLARE_TOGGLE_COMMAND(ToggleRequestsTreeViewVisibility)
-#undef DECLARE_TOGGLE_COMMAND
+	INSIGHTS_DECLARE_TOGGLE_COMMAND(ToggleTimingViewVisibility)
+	INSIGHTS_DECLARE_TOGGLE_COMMAND(ToggleEventAggregationTreeViewVisibility)
+	INSIGHTS_DECLARE_TOGGLE_COMMAND(ToggleObjectTypeAggregationTreeViewVisibility)
+	INSIGHTS_DECLARE_TOGGLE_COMMAND(TogglePackageDetailsTreeViewVisibility)
+	INSIGHTS_DECLARE_TOGGLE_COMMAND(ToggleExportDetailsTreeViewVisibility)
+	INSIGHTS_DECLARE_TOGGLE_COMMAND(ToggleRequestsTreeViewVisibility)
 
 	//////////////////////////////////////////////////
 
diff --git a/Engine/Source/Developer/TraceInsights/Private/Insights/LoadingProfiler/LoadingProfilerManager.cpp b/Engine/Source/Developer/TraceInsights/Private/Insights/LoadingProfiler/LoadingProfilerManager.cpp
--- a/Engine/Source/Developer/TraceInsights/Private/Insights/LoadingProfiler/LoadingProfilerManager.cpp
+++ b/Engine/Source/Developer/TraceInsights/Private/Insights/LoadingProfiler/LoadingProfilerManager.cpp
@@ -129,12 +129,12 @@ FLoadingProfilerManager::~FLoadingProfilerManager()
 
 void FLoadingProfilerManager::BindCommands()
 {
-	ActionManager.Map_ToggleTimingViewVisibility_Global();
-	ActionManager.Map_ToggleEventAggregationTreeViewVisibility_Global();
-	ActionManager.Map_ToggleObjectTypeAggregationTreeViewVisibility_Global();
-	ActionManager.Map_TogglePackageDetailsTreeViewVisibility_Global();
-	ActionManager.Map_ToggleExportDetailsTreeViewVisibility_Global();
-	ActionManager.Map_ToggleRequestsTreeViewVisibility_Global();
+	ActionManager.Map_ToggleTimingViewVisibility();
+	ActionManager.Map_ToggleEventAggregationTreeViewVisibility();
+	ActionManager.Map_ToggleObjectTypeAggregationTreeViewVisibility();
+	ActionManager.Map_TogglePackageDetailsTreeViewVisibility();
+	ActionManager.Map_ToggleExportDetailsTreeViewVisibility();
+	ActionManager.Map_ToggleRequestsTreeViewVisibility();
 }
 
 ////////////////////////////////////////////////////////////////////////////////////////////////////
diff --git a/Engine/Source/Developer/TraceInsights/Private/Insights/MemoryProfiler/MemoryProfilerCommands.cpp b/Engine/Source/Developer/TraceInsights/Private/Insights/MemoryProfiler/MemoryProfilerCommands.cpp
--- a/Engine/Source/Developer/TraceInsights/Private/Insights/MemoryProfiler/MemoryProfilerCommands.cpp
+++ b/Engine/Source/Developer/TraceInsights/Private/Insights/MemoryProfiler/MemoryProfilerCommands.cpp
@@ -4,7 +4,6 @@
 
 #include "DesktopPlatformModule.h"
 #include "Framework/Application/SlateApplication.h"
-#include "Framework/MultiBox/MultiBoxBuilder.h"
 
 // TraceInsights
 #include "Insights/InsightsManager.h"
@@ -18,23 +17,6 @@
 namespace UE::Insights::MemoryProfiler
 {
 
-////////////////////////////////////////////////////////////////////////////////////////////////////
-// FMemoryProfilerMenuBuilder
-////////////////////////////////////////////////////////////////////////////////////////////////////
-
-void FMemoryProfilerMenuBuilder::AddMenuEntry(FMenuBuilder& MenuBuilder, const TSharedPtr< FUICommandInfo >& UICommandInfo, const FUIAction& UIAction)
-{
-	MenuBuilder.AddMenuEntry
-	(
-		UICommandInfo->GetLabel(),
-		UICommandInfo->GetDescription(),
-		UICommandInfo->GetIcon(),
-		UIAction,
-		NAME_None,
-		UICommandInfo->GetUserInterfaceType()
-	);
-}
-
 ////////////////////////////////////////////////////////////////////////////////////////////////////
 // FMemoryProfilerCommands
 ////////////////////////////////////////////////////////////////////////////////////////////////////
@@ -90,45 +72,10 @@ UE_ENABLE_OPTIMIZATION_SHIP
 // Toggle Commands
 ////////////////////////////////////////////////////////////////////////////////////////////////////
 
-#define IMPLEMENT_TOGGLE_COMMAND(CmdName, IsEnabled, SetIsEnabled) \
-	\
-	void FMemoryProfilerActionManager::Map_##CmdName##_Global()\
-	{\
-		This->CommandList->MapAction(This->GetCommands().CmdName, CmdName##_Custom());\
-	}\
-	\
-	const FUIAction FMemoryProfilerActionManager::CmdName##_Custom() \
-	{\
-		FUIAction UIAction;\
-		UIAction.ExecuteAction = FExecuteAction::CreateRaw(this, &FMemoryProfilerActionManager::CmdName##_Execute);\
-		UIAction.CanExecuteAction = FCanExecuteAction::CreateRaw(this, &FMemoryProfilerActionManager::CmdName##_CanExecute);\
-		UIAction.GetActionCheckState = FGetActionCheckState::CreateRaw(this, &FMemoryProfilerActionManager::CmdName##_GetCheckState);\
-		return UIAction;\
-	}\
-	\
-	void FMemoryProfilerActionManager::CmdName##_Execute()\
-	{\
-		const bool b##IsEnabled = !This->IsEnabled();\
-		This->SetIsEnabled(b##IsEnabled);\
-	}\
-	\
-	bool FMemoryProfilerActionManager::CmdName##_CanExecute() const\
-	{\
-		return FInsightsManager::Get()->GetSession().IsValid();\
-	}\
-	\
-	ECheckBoxState FMemoryProfilerActionManager::CmdName##_GetCheckState() const\
-	{\
-		const bool b##IsEnabled = This->IsEnabled();\
-		return b##IsEnabled ? ECheckBoxState::Checked : ECheckBoxState::Unchecked;\
-	}
-
-IMPLEMENT_TOGGLE_COMMAND(ToggleTimingViewVisibility, IsTimingViewVisible, ShowHideTimingView)
-IMPLEMENT_TOGGLE_COMMAND(ToggleMemInvestigationViewVisibility, IsMemInvestigationViewVisible, ShowHideMemInvestigationView)
-IMPLEMENT_TOGGLE_COMMAND(ToggleMemTagTreeViewVisibility, IsMemTagTreeViewVisible, ShowHideMemTagTreeView)
-IMPLEMENT_TOGGLE_COMMAND(ToggleModulesViewVisibility, IsModulesViewVisible, ShowHideModulesView)
-
-#undef IMPLEMENT_TOGGLE_COMMAND
+INSIGHTS_IMPLEMENT_TOGGLE_COMMAND(FMemoryProfilerActionManager, This, ToggleTimingViewVisibility, IsTimingViewVisible, ShowHideTimingView)
+INSIGHTS_IMPLEMENT_TOGGLE_COMMAND(FMemoryProfilerActionManager, This, ToggleMemInvestigationViewVisibility, IsMemInvestigationViewVisible, ShowHideMemInvestigationView)
+INSIGHTS_IMPLEMENT_TOGGLE_COMMAND(FMemoryProfilerActionManager, This, ToggleMemTagTreeViewVisibility, IsMemTagTreeViewVisible, ShowHideMemTagTreeView)
+INSIGHTS_IMPLEMENT_TOGGLE_COMMAND(FMemoryProfilerActionManager, This, ToggleModulesViewVisibility, IsModulesViewVisible, ShowHideModulesView)
 
 ////////////////////////////////////////////////////////////////////////////////////////////////////
 
diff --git a/Engine/Source/Developer/TraceInsights/Private/Insights/MemoryProfiler/MemoryProfilerCommands.h b/Engine/Source/Developer/TraceInsights/Private/Insights/MemoryProfiler/MemoryProfilerCommands.h
--- a/Engine/Source/Developer/TraceInsights/Private/Insights/MemoryProfiler/MemoryProfilerCommands.h
+++ b/Engine/Source/Developer/TraceInsights/Private/Insights/MemoryProfiler/MemoryProfilerCommands.h
@@ -7,7 +7,8 @@
 #include "Framework/Commands/UIAction.h"
 #include "Framework/Commands/Commands.h"
 
-class FMenuBuilder;
+// TraceInsights
+#include "Insights/Common/CommandUtils.h"
 
 namespace UE::Insights::MemoryProfiler
 {
@@ -24,48 +25,19 @@ public:
 	virtual void RegisterCommands() override;
 
 public:
-	//////////////////////////////////////////////////
-	// Global commands need to implement following method:
-	//     void Map_<CommandName>_Global();
-	// Custom commands needs to implement also the following method:
-	//     const FUIAction <CommandName>_Custom(...) const;
-	//////////////////////////////////////////////////
-
-	/** Toggles visibility for the Timing view. Global and custom command. */
+	/** Toggles visibility for the Timing view. */
 	TSharedPtr<FUICommandInfo> ToggleTimingViewVisibility;
 
-	/** Toggles visibility for the Memory Investigation view. Global and custom command. */
+	/** Toggles visibility for the Memory Investigation view. */
 	TSharedPtr<FUICommandInfo> ToggleMemInvestigationViewVisibility;
 
-	/** Toggles visibility for the Memory Tags tree view. Global and custom command. */
+	/** Toggles visibility for the Memory Tags tree view. */
 	TSharedPtr<FUICommandInfo> ToggleMemTagTreeViewVisibility;
 
-	/** Toggles visibility for the Modules view. Global and custom command. */
+	/** Toggles visibility for the Modules view. */
 	TSharedPtr<FUICommandInfo> ToggleModulesViewVisibility;
 };
 
-////////////////////////////////////////////////////////////////////////////////////////////////////
-/**
- * Menu builder. Helper class for adding a customized menu entry using the global UI command info.
- */
-class FMemoryProfilerMenuBuilder
-{
-public:
-	/**
-	 * Helper method for adding a customized menu entry using the global UI command info.
-	 * FUICommandInfo cannot be executed with custom parameters, so we need to create a custom FUIAction,
-	 * but sometime we have global and local version for the UI command, so reuse data from the global UI command info.
-	 * Ex:
-	 *     SessionInstance_ToggleCapture          - Global version will toggle capture process for all active session instances
-	 *     SessionInstance_ToggleCapture_OneParam - Local version will toggle capture process only for the specified session instance
-	 *
-	 * @param MenuBuilder The menu to add items to
-	 * @param FUICommandInfo A shared pointer to the UI command info
-	 * @param UIAction Customized version of the UI command info stored in an UI action
-	 */
-	static void AddMenuEntry(FMenuBuilder& MenuBuilder, const TSharedPtr<FUICommandInfo>& UICommandInfo, const FUIAction& UIAction);
-};
-
 ////////////////////////////////////////////////////////////////////////////////////////////////////
 /**
  * Class that provides helper functions for the commands to avoid cluttering profiler manager with many small functions.
@@ -84,20 +56,10 @@ private:
 	//////////////////////////////////////////////////
 	// Toggle Commands
 
-#define DECLARE_TOGGLE_COMMAND(CmdName)\
-public:\
-	void Map_##CmdName##_Global(); /**< Maps UI command info CmdName with the specified UI command list. */\
-	const FUIAction CmdName##_Custom(); /**< UI action for CmdName command. */\
-protected:\
-	void CmdName##_Execute(); /**< Handles FExecuteAction for CmdName. */\
-	bool CmdName##_CanExecute() const; /**< Handles FCanExecuteAction for CmdName. */\
-	ECheckBoxState CmdName##_GetCheckState() const; /**< Handles FGetActionCheckState for CmdName. */
-
-	DECLARE_TOGGLE_COMMAND(ToggleTimingViewVisibility)
-	DECLARE_TOGGLE_COMMAND(ToggleMemInvestigationViewVisibility)
-	DECLARE_TOGGLE_COMMAND(ToggleMemTagTreeViewVisibility)
-	DECLARE_TOGGLE_COMMAND(ToggleModulesViewVisibility)
-#undef DECLARE_TOGGLE_COMMAND
+	INSIGHTS_DECLARE_TOGGLE_COMMAND(ToggleTimingViewVisibility)
+	INSIGHTS_DECLARE_TOGGLE_COMMAND(ToggleMemInvestigationViewVisibility)
+	INSIGHTS_DECLARE_TOGGLE_COMMAND(ToggleMemTagTreeViewVisibility)
+	INSIGHTS_DECLARE_TOGGLE_COMMAND(ToggleModulesViewVisibility)
 
 	//////////////////////////////////////////////////
 
diff --git a/Engine/Source/Developer/TraceInsights/Private/Insights/MemoryProfiler/MemoryProfilerManager.cpp b/Engine/Source/Developer/TraceInsights/Private/Insights/MemoryProfiler/MemoryProfilerManager.cpp
--- a/Engine/Source/Developer/TraceInsights/Private/Insights/MemoryProfiler/MemoryProfilerManager.cpp
+++ b/Engine/Source/Developer/TraceInsights/Private/Insights/MemoryProfiler/MemoryProfilerManager.cpp
@@ -144,10 +144,10 @@ FMemoryProfilerManager::~FMemoryProfilerManager()
 
 void FMemoryProfilerManager::BindCommands()
 {
-	ActionManager.Map_ToggleTimingViewVisibility_Global();
-	ActionManager.Map_ToggleMemInvestigationViewVisibility_Global();
-	ActionManager.Map_ToggleMemTagTreeViewVisibility_Global();
-	ActionManager.Map_ToggleModulesViewVisibility_Global();
+	ActionManager.Map_ToggleTimingViewVisibility();
+	ActionManager.Map_ToggleMemInvestigationViewVisibility();
+	ActionManager.Map_ToggleMemTagTreeViewVisibility();
+	ActionManager.Map_ToggleModulesViewVisibility();
 }
 
 ////////////////////////////////////////////////////////////////////////////////////////////////////
diff --git a/Engine/Source/Developer/TraceInsights/Private/Insights/NetworkingProfiler/NetworkingProfilerCommands.cpp b/Engine/Source/Developer/TraceInsights/Private/Insights/NetworkingProfiler/NetworkingProfilerCommands.cpp
--- a/Engine/Source/Developer/TraceInsights/Private/Insights/NetworkingProfiler/NetworkingProfilerCommands.cpp
+++ b/Engine/Source/Developer/TraceInsights/Private/Insights/NetworkingProfiler/NetworkingProfilerCommands.cpp
@@ -4,7 +4,6 @@
 
 #include "DesktopPlatformModule.h"
 #include "Framework/Application/SlateApplication.h"
-#include "Framework/MultiBox/MultiBoxBuilder.h"
 
 // TraceInsights
 #include "Insights/InsightsManager.h"
@@ -18,108 +17,56 @@
 namespace UE::Insights::NetworkingProfiler
 {
 
-	////////////////////////////////////////////////////////////////////////////////////////////////////
-	// FNetworkingProfilerMenuBuilder
-	////////////////////////////////////////////////////////////////////////////////////////////////////
-
-	void FNetworkingProfilerMenuBuilder::AddMenuEntry(FMenuBuilder& MenuBuilder, const TSharedPtr< FUICommandInfo >& UICommandInfo, const FUIAction& UIAction)
-	{
-		MenuBuilder.AddMenuEntry
-		(
-			UICommandInfo->GetLabel(),
-			UICommandInfo->GetDescription(),
-			UICommandInfo->GetIcon(),
-			UIAction,
-			NAME_None,
-			UICommandInfo->GetUserInterfaceType()
-		);
-	}
-
-	////////////////////////////////////////////////////////////////////////////////////////////////////
-	// FNetworkingProfilerCommands
-	////////////////////////////////////////////////////////////////////////////////////////////////////
-
-	FNetworkingProfilerCommands::FNetworkingProfilerCommands()
-		: TCommands<FNetworkingProfilerCommands>(
-			TEXT("NetworkingProfilerCommands"),
-			NSLOCTEXT("Contexts", "NetworkingProfilerCommands", "Insights - Networking Insights"),
-			NAME_None,
-			FInsightsStyle::GetStyleSetName())
-	{
-	}
-
-	////////////////////////////////////////////////////////////////////////////////////////////////////
-
-	// UI_COMMAND takes long for the compiler to optimize
-	UE_DISABLE_OPTIMIZATION_SHIP
-		void FNetworkingProfilerCommands::RegisterCommands()
-	{
-		UI_COMMAND(TogglePacketViewVisibility,
-			"Packets",
-			"Toggles the visibility of the Packets view.",
-			EUserInterfaceActionType::ToggleButton,
-			FInputChord());
-
-		UI_COMMAND(TogglePacketContentViewVisibility,
-			"Packet Content",
-			"Toggles the visibility of the Packet Content view.",
-			EUserInterfaceActionType::ToggleButton,
-			FInputChord());
-
-		UI_COMMAND(ToggleNetStatsViewVisibility,
-			"Net Stats",
-			"Toggles the visibility of the Net Stats view.",
-			EUserInterfaceActionType::ToggleButton,
-			FInputChord());
-
-		UI_COMMAND(ToggleNetStatsCountersViewVisibility,
-			"Net Stats Counters",
-			"Toggles the visibility of the Net Stats view.",
-			EUserInterfaceActionType::ToggleButton,
-			FInputChord());
-	}
-	UE_ENABLE_OPTIMIZATION_SHIP
-
-		////////////////////////////////////////////////////////////////////////////////////////////////////
-		// Toggle Commands
-		////////////////////////////////////////////////////////////////////////////////////////////////////
-
-#define IMPLEMENT_TOGGLE_COMMAND(CmdName, IsEnabled, SetIsEnabled) \
-	\
-	void FNetworkingProfilerActionManager::Map_##CmdName##_Global()\
-	{\
-		This->CommandList->MapAction(This->GetCommands().CmdName, CmdName##_Custom());\
-	}\
-	\
-	const FUIAction FNetworkingProfilerActionManager::CmdName##_Custom() \
-	{\
-		FUIAction UIAction;\
-		UIAction.ExecuteAction = FExecuteAction::CreateRaw(this, &FNetworkingProfilerActionManager::CmdName##_Execute);\
-		UIAction.CanExecuteAction = FCanExecuteAction::CreateRaw(this, &FNetworkingProfilerActionManager::CmdName##_CanExecute);\
-		UIAction.GetActionCheckState = FGetActionCheckState::CreateRaw(this, &FNetworkingProfilerActionManager::CmdName##_GetCheckState);\
-		return UIAction;\
-	}\
-	\
-	void FNetworkingProfilerActionManager::CmdName##_Execute()\
-	{\
-		const bool b##IsEnabled = !This->IsEnabled();\
-		This->SetIsEnabled(b##IsEnabled);\
-	}\
-	\
-	bool FNetworkingProfilerActionManager::CmdName##_CanExecute() const\
-	{\
-		return FInsightsManager::Get()->GetSession().IsValid();\
-	}\
-	\
-	ECheckBoxState FNetworkingProfilerActionManager::CmdName##_GetCheckState() const\
-	{\
-		const bool b##IsEnabled = This->IsEnabled();\
-		return b##IsEnabled ? ECheckBoxState::Checked : ECheckBoxState::Unchecked;\
-	}
-
-//IMPLEMENT_TOGGLE_COMMAND(ToggleAAAViewVisibility, IsAAAViewVisible, ShowHideAAAView)
-
-#undef IMPLEMENT_TOGGLE_COMMAND
+////////////////////////////////////////////////////////////////////////////////////////////////////
+// FNetworkingProfilerCommands
+////////////////////////////////////////////////////////////////////////////////////////////////////
+
+FNetworkingProfilerCommands::FNetworkingProfilerCommands()
+	: TCommands<FNetworkingProfilerCommands>(
+		TEXT("NetworkingProfilerCommands"),
+		NSLOCTEXT("Contexts", "NetworkingProfilerCommands", "Insights - Networking Insights"),
+		NAME_None,
+		FInsightsStyle::GetStyleSetName())
+{
+}
+
+////////////////////////////////////////////////////////////////////////////////////////////////////
+
+// UI_COMMAND takes long for the compiler to optimize
+UE_DISABLE_OPTIMIZATION_SHIP
+void FNetworkingProfilerCommands::RegisterCommands()
+{
+	UI_COMMAND(TogglePacketViewVisibility,
+		"Packets",
+		"Toggles the visibility of the Packets view.",
+		EUserInterfaceActionType::ToggleButton,
+		FInputChord());
+
+	UI_COMMAND(TogglePacketContentViewVisibility,
+		"Packet Content",
+		"Toggles the visibility of the Packet Content view.",
+		EUserInterfaceActionType::ToggleButton,
+		FInputChord());
+
+	UI_COMMAND(ToggleNetStatsViewVisibility,
+		"Net Stats",
+		"Toggles the visibility of the Net Stats view.",
+		EUserInterfaceActionType::ToggleButton,
+		FInputChord());
+
+	UI_COMMAND(ToggleNetStatsCountersViewVisibility,
+		"Net Stats Counters",
+		"Toggles the visibility of the Net Stats view.",
+		EUserInterfaceActionType::ToggleButton,
+		FInputChord());
+}
+UE_ENABLE_OPTIMIZATION_SHIP
+
+////////////////////////////////////////////////////////////////////////////////////////////////////
+// Toggle Commands
+////////////////////////////////////////////////////////////////////////////////////////////////////
+
+//INSIGHTS_IMPLEMENT_TOGGLE_COMMAND(FNetworkingProfilerActionManager, This, ToggleAAAViewVisibility, IsAAAViewVisible, ShowHideAAAView)
 
 ////////////////////////////////////////////////////////////////////////////////////////////////////
 
diff --git a/Engine/Source/Developer/TraceInsights/Private/Insights/NetworkingProfiler/NetworkingProfilerCommands.h b/Engine/Source/Developer/TraceInsights/Private/Insights/NetworkingProfiler/NetworkingProfilerCommands.h
--- a/Engine/Source/Developer/TraceInsights/Private/Insights/NetworkingProfiler/NetworkingProfilerCommands.h
+++ b/Engine/Source/Developer/TraceInsights/Private/Insights/NetworkingProfiler/NetworkingProfilerCommands.h
@@ -7,7 +7,8 @@
 #include "Framework/Commands/UIAction.h"
 #include "Framework/Commands/Commands.h"
 
-class FMenuBuilder;
+// TraceInsights
+#include "Insights/Common/CommandUtils.h"
 
 namespace UE::Insights::NetworkingProfiler
 {
@@ -26,13 +27,6 @@ public:
 	virtual void RegisterCommands() override;
 
 public:
-	//////////////////////////////////////////////////
-	// Global commands need to implement following method:
-	//     void Map_<CommandName>_Global();
-	// Custom commands needs to implement also the following method:
-	//     const FUIAction <CommandName>_Custom(...) const;
-	//////////////////////////////////////////////////
-
 	/** Toggles visibility for the Packet view. Per profiler window, custom command. */
 	TSharedPtr<FUICommandInfo> TogglePacketViewVisibility;
 
@@ -46,28 +40,6 @@ public:
 	TSharedPtr<FUICommandInfo> ToggleNetStatsCountersViewVisibility;
 };
 
-////////////////////////////////////////////////////////////////////////////////////////////////////
-/**
- * Menu builder. Helper class for adding a customized menu entry using the global UI command info.
- */
-class FNetworkingProfilerMenuBuilder
-{
-public:
-	/**
-	 * Helper method for adding a customized menu entry using the global UI command info.
-	 * FUICommandInfo cannot be executed with custom parameters, so we need to create a custom FUIAction,
-	 * but sometime we have global and local version for the UI command, so reuse data from the global UI command info.
-	 * Ex:
-	 *     SessionInstance_ToggleCapture          - Global version will toggle capture process for all active session instances
-	 *     SessionInstance_ToggleCapture_OneParam - Local version will toggle capture process only for the specified session instance
-	 *
-	 * @param MenuBuilder The menu to add items to
-	 * @param FUICommandInfo A shared pointer to the UI command info
-	 * @param UIAction Customized version of the UI command info stored in an UI action
-	 */
-	static void AddMenuEntry(FMenuBuilder& MenuBuilder, const TSharedPtr<FUICommandInfo>& UICommandInfo, const FUIAction& UIAction);
-};
-
 ////////////////////////////////////////////////////////////////////////////////////////////////////
 /**
  * Class that provides helper functions for the commands to avoid cluttering profiler manager with many small functions.
@@ -86,17 +58,7 @@ private:
 	//////////////////////////////////////////////////
 	// Toggle Commands
 
-#define DECLARE_TOGGLE_COMMAND(CmdName)\
-public:\
-	void Map_##CmdName##_Global(); /**< Maps UI command info CmdName with the specified UI command list. */\
-	const FUIAction CmdName##_Custom(); /**< UI action for CmdName command. */\
-protected:\
-	void CmdName##_Execute(); /**< Handles FExecuteAction for CmdName. */\
-	bool CmdName##_CanExecute() const; /**< Handles FCanExecuteAction for CmdName. */\
-	ECheckBoxState CmdName##_GetCheckState() const; /**< Handles FGetActionCheckState for CmdName. */
-
-	//DECLARE_TOGGLE_COMMAND(ToggleAAAViewVisibility)
-#undef DECLARE_TOGGLE_COMMAND
+	//INSIGHTS_DECLARE_TOGGLE_COMMAND(ToggleAAAViewVisibility)
 
 	//////////////////////////////////////////////////
 
diff --git a/Engine/Source/Developer/TraceInsights/Private/Insights/NetworkingProfiler/Widgets/SNetworkingProfilerWindow.cpp b/Engine/Source/Developer/TraceInsights/Private/Insights/NetworkingProfiler/Widgets/SNetworkingProfilerWindow.cpp
--- a/Engine/Source/Developer/TraceInsights/Private/Insights/NetworkingProfiler/Widgets/SNetworkingProfilerWindow.cpp
+++ b/Engine/Source/Developer/TraceInsights/Private/Insights/NetworkingProfiler/Widgets/SNetworkingProfilerWindow.cpp
@@ -305,7 +305,7 @@ TSharedRef<FTabManager::FLayout> SNetworkingProfilerWindow::CreateDefaultTabLayo
 				->Split
 				(
 					FTabManager::NewStack()
-					->SetSizeCoefficient(0.35f)					
+					->SetSizeCoefficient(0.35f)
 					->AddTab(FNetworkingProfilerTabs::PacketViewID, ETabState::OpenedTab)
 					->SetHideTabWell(true)
 				)
@@ -342,6 +342,13 @@ TSharedRef<SWidget> SNetworkingProfilerWindow::CreateToolbar(TSharedPtr<FExtende
 
 ////////////////////////////////////////////////////////////////////////////////////////////////////
 
+const FNetworkingProfilerCommands& SNetworkingProfilerWindow::GetCommands() const
+{
+	return FNetworkingProfilerManager::GetCommands();
+}
+
+////////////////////////////////////////////////////////////////////////////////////////////////////
+
 void SNetworkingProfilerWindow::BindCommands()
 {
 	check(GetCommandList().IsValid());
@@ -351,50 +358,17 @@ void SNetworkingProfilerWindow::BindCommands()
 	Map_ToggleNetStatsViewVisibility();
 	Map_ToggleNetStatsCountersViewVisibility();
 
-	GetCommandList()->MapAction(FInsightsCommands::Get().ToggleDebugInfo, FInsightsManager::GetActionManager().ToggleDebugInfo_Custom());
+	GetCommandList()->MapAction(FInsightsCommands::Get().ToggleDebugInfo, FInsightsManager::GetActionManager().GetAction_ToggleDebugInfo());
 }
 
 ////////////////////////////////////////////////////////////////////////////////////////////////////
 // Toggle Commands
 ////////////////////////////////////////////////////////////////////////////////////////////////////
 
-#define IMPLEMENT_TOGGLE_COMMAND(CmdName, IsEnabled, SetEnabled) \
-	\
-	void SNetworkingProfilerWindow::Map_##CmdName()\
-	{\
-		GetCommandList()->MapAction(FNetworkingProfilerManager::GetCommands().CmdName, CmdName##_Custom());\
-	}\
-	\
-	const FUIAction SNetworkingProfilerWindow::CmdName##_Custom() \
-	{\
-		FUIAction UIAction;\
-		UIAction.ExecuteAction = FExecuteAction::CreateRaw(this, &SNetworkingProfilerWindow::CmdName##_Execute);\
-		UIAction.CanExecuteAction = FCanExecuteAction::CreateRaw(this, &SNetworkingProfilerWindow::CmdName##_CanExecute);\
-		UIAction.GetActionCheckState = FGetActionCheckState::CreateRaw(this, &SNetworkingProfilerWindow::CmdName##_GetCheckState);\
-		return UIAction;\
-	}\
-	\
-	void SNetworkingProfilerWindow::CmdName##_Execute()\
-	{\
-		SetEnabled(!IsEnabled());\
-	}\
-	\
-	bool SNetworkingProfilerWindow::CmdName##_CanExecute() const\
-	{\
-		return FInsightsManager::Get()->GetSession().IsValid();\
-	}\
-	\
-	ECheckBoxState SNetworkingProfilerWindow::CmdName##_GetCheckState() const\
-	{\
-		return IsEnabled() ? ECheckBoxState::Checked : ECheckBoxState::Unchecked;\
-	}
-
-IMPLEMENT_TOGGLE_COMMAND(TogglePacketViewVisibility, IsPacketViewVisible, ShowOrHidePacketView)
-IMPLEMENT_TOGGLE_COMMAND(TogglePacketContentViewVisibility, IsPacketContentViewVisible, ShowOrHidePacketContentView)
-IMPLEMENT_TOGGLE_COMMAND(ToggleNetStatsViewVisibility, IsNetStatsViewVisible, ShowOrHideNetStatsView)
-IMPLEMENT_TOGGLE_COMMAND(ToggleNetStatsCountersViewVisibility, IsNetStatsCountersViewVisible, ShowOrHideNetStatsCountersView)
-
-#undef IMPLEMENT_TOGGLE_COMMAND
+INSIGHTS_IMPLEMENT_TOGGLE_COMMAND(SNetworkingProfilerWindow, this, TogglePacketViewVisibility, IsPacketViewVisible, ShowOrHidePacketView)
+INSIGHTS_IMPLEMENT_TOGGLE_COMMAND(SNetworkingProfilerWindow, this, TogglePacketContentViewVisibility, IsPacketContentViewVisible, ShowOrHidePacketContentView)
+INSIGHTS_IMPLEMENT_TOGGLE_COMMAND(SNetworkingProfilerWindow, this, ToggleNetStatsViewVisibility, IsNetStatsViewVisible, ShowOrHideNetStatsView)
+INSIGHTS_IMPLEMENT_TOGGLE_COMMAND(SNetworkingProfilerWindow, this, ToggleNetStatsCountersViewVisibility, IsNetStatsCountersViewVisible, ShowOrHideNetStatsCountersView)
 
 ////////////////////////////////////////////////////////////////////////////////////////////////////
 
diff --git a/Engine/Source/Developer/TraceInsights/Private/Insights/NetworkingProfiler/Widgets/SNetworkingProfilerWindow.h b/Engine/Source/Developer/TraceInsights/Private/Insights/NetworkingProfiler/Widgets/SNetworkingProfilerWindow.h
--- a/Engine/Source/Developer/TraceInsights/Private/Insights/NetworkingProfiler/Widgets/SNetworkingProfilerWindow.h
+++ b/Engine/Source/Developer/TraceInsights/Private/Insights/NetworkingProfiler/Widgets/SNetworkingProfilerWindow.h
@@ -9,6 +9,7 @@
 #include "TraceServices/Model/NetProfiler.h"
 
 // TraceInsights
+#include "Insights/Common/CommandUtils.h"
 #include "Insights/Widgets/SMajorTabWindow.h"
 
 namespace UE::Insights::NetworkingProfiler
@@ -20,6 +21,7 @@ class SPacketView;
 class SPacketContentView;
 class SNetStatsView;
 class SNetStatsCountersView;
+class FNetworkingProfilerCommands;
 
 ////////////////////////////////////////////////////////////////////////////////////////////////////
 
@@ -142,25 +144,16 @@ private:
 
 	//////////////////////////////////////////////////
 
+	const FNetworkingProfilerCommands& GetCommands() const;
 	void BindCommands();
 
 	//////////////////////////////////////////////////
 	// Toggle Commands
 
-#define DECLARE_TOGGLE_COMMAND(CmdName)\
-public:\
-	void Map_##CmdName(); /**< Maps UI command info CmdName with the specified UI command list. */\
-	const FUIAction CmdName##_Custom(); /**< UI action for CmdName command. */\
-private:\
-	void CmdName##_Execute(); /**< Handles FExecuteAction for CmdName. */\
-	bool CmdName##_CanExecute() const; /**< Handles FCanExecuteAction for CmdName. */\
-	ECheckBoxState CmdName##_GetCheckState() const; /**< Handles FGetActionCheckState for CmdName. */
-
-	DECLARE_TOGGLE_COMMAND(TogglePacketViewVisibility)
-	DECLARE_TOGGLE_COMMAND(TogglePacketContentViewVisibility)
-	DECLARE_TOGGLE_COMMAND(ToggleNetStatsViewVisibility)
-	DECLARE_TOGGLE_COMMAND(ToggleNetStatsCountersViewVisibility)
-#undef DECLARE_TOGGLE_COMMAND
+	INSIGHTS_DECLARE_TOGGLE_COMMAND(TogglePacketViewVisibility)
+	INSIGHTS_DECLARE_TOGGLE_COMMAND(TogglePacketContentViewVisibility)
+	INSIGHTS_DECLARE_TOGGLE_COMMAND(ToggleNetStatsViewVisibility)
+	INSIGHTS_DECLARE_TOGGLE_COMMAND(ToggleNetStatsCountersViewVisibility)
 
 	//////////////////////////////////////////////////
 
@@ -212,7 +205,7 @@ private:
 
 	/** The NetStatsCounters widget */
 	TSharedPtr<SNetStatsCountersView> NetStatsCountersView;
-	
+
 	TSharedPtr<SComboBox<TSharedPtr<FGameInstanceItem>>> GameInstanceComboBox;
 	TArray<TSharedPtr<FGameInstanceItem>> AvailableGameInstances;
 	TSharedPtr<FGameInstanceItem> SelectedGameInstance;
diff --git a/Engine/Source/Developer/TraceInsights/Private/Insights/ObjectProfiler/ObjectProfilerCommands.cpp b/Engine/Source/Developer/TraceInsights/Private/Insights/ObjectProfiler/ObjectProfilerCommands.cpp
--- a/Engine/Source/Developer/TraceInsights/Private/Insights/ObjectProfiler/ObjectProfilerCommands.cpp
+++ b/Engine/Source/Developer/TraceInsights/Private/Insights/ObjectProfiler/ObjectProfilerCommands.cpp
@@ -60,43 +60,8 @@ UE_ENABLE_OPTIMIZATION_SHIP
 // Toggle Commands
 ////////////////////////////////////////////////////////////////////////////////////////////////////
 
-#define IMPLEMENT_TOGGLE_COMMAND(CmdName, IsEnabled, SetIsEnabled) \
-	\
-	void FObjectProfilerActionManager::Map_##CmdName##_Global()\
-	{\
-		This->CommandList->MapAction(This->GetCommands().CmdName, CmdName##_Custom());\
-	}\
-	\
-	const FUIAction FObjectProfilerActionManager::CmdName##_Custom() \
-	{\
-		FUIAction UIAction;\
-		UIAction.ExecuteAction = FExecuteAction::CreateRaw(this, &FObjectProfilerActionManager::CmdName##_Execute);\
-		UIAction.CanExecuteAction = FCanExecuteAction::CreateRaw(this, &FObjectProfilerActionManager::CmdName##_CanExecute);\
-		UIAction.GetActionCheckState = FGetActionCheckState::CreateRaw(this, &FObjectProfilerActionManager::CmdName##_GetCheckState);\
-		return UIAction;\
-	}\
-	\
-	void FObjectProfilerActionManager::CmdName##_Execute()\
-	{\
-		const bool b##IsEnabled = !This->IsEnabled();\
-		This->SetIsEnabled(b##IsEnabled);\
-	}\
-	\
-	bool FObjectProfilerActionManager::CmdName##_CanExecute() const\
-	{\
-		return FInsightsManager::Get()->GetSession().IsValid();\
-	}\
-	\
-	ECheckBoxState FObjectProfilerActionManager::CmdName##_GetCheckState() const\
-	{\
-		const bool b##IsEnabled = This->IsEnabled();\
-		return b##IsEnabled ? ECheckBoxState::Checked : ECheckBoxState::Unchecked;\
-	}
-
-IMPLEMENT_TOGGLE_COMMAND(ToggleObjectTableTreeViewVisibility, IsObjectTableTreeViewVisible, ShowHideObjectTableTreeView)
-IMPLEMENT_TOGGLE_COMMAND(ToggleObjectDetailsViewVisibility, IsObjectDetailsViewVisible, ShowHideObjectDetailsView)
-
-#undef IMPLEMENT_TOGGLE_COMMAND
+INSIGHTS_IMPLEMENT_TOGGLE_COMMAND(FObjectProfilerActionManager, This, ToggleObjectTableTreeViewVisibility, IsObjectTableTreeViewVisible, ShowHideObjectTableTreeView)
+INSIGHTS_IMPLEMENT_TOGGLE_COMMAND(FObjectProfilerActionManager, This, ToggleObjectDetailsViewVisibility, IsObjectDetailsViewVisible, ShowHideObjectDetailsView)
 
 ////////////////////////////////////////////////////////////////////////////////////////////////////
 
diff --git a/Engine/Source/Developer/TraceInsights/Private/Insights/ObjectProfiler/ObjectProfilerCommands.h b/Engine/Source/Developer/TraceInsights/Private/Insights/ObjectProfiler/ObjectProfilerCommands.h
--- a/Engine/Source/Developer/TraceInsights/Private/Insights/ObjectProfiler/ObjectProfilerCommands.h
+++ b/Engine/Source/Developer/TraceInsights/Private/Insights/ObjectProfiler/ObjectProfilerCommands.h
@@ -7,6 +7,9 @@
 #include "Framework/Commands/UIAction.h"
 #include "Framework/Commands/Commands.h"
 
+// TraceInsights
+#include "Insights/Common/CommandUtils.h"
+
 namespace UE::Insights::ObjectProfiler
 {
 
@@ -22,17 +25,10 @@ public:
 	virtual void RegisterCommands() override;
 
 public:
-	//////////////////////////////////////////////////
-	// Global commands need to implement following method:
-	//     void Map_<CommandName>_Global();
-	// Custom commands needs to implement also the following method:
-	//     const FUIAction <CommandName>_Custom(...) const;
-	//////////////////////////////////////////////////
-
-	/** Toggles visibility for the Objects table/tree view. Global and custom command. */
+	/** Toggles visibility for the Objects table/tree view. */
 	TSharedPtr<FUICommandInfo> ToggleObjectTableTreeViewVisibility;
 
-	/** Toggles visibility for the Object Details view. Global and custom command. */
+	/** Toggles visibility for the Object Details view. */
 	TSharedPtr<FUICommandInfo> ToggleObjectDetailsViewVisibility;
 };
 
@@ -54,18 +50,8 @@ private:
 	//////////////////////////////////////////////////
 	// Toggle Commands
 
-#define DECLARE_TOGGLE_COMMAND(CmdName)\
-public:\
-	void Map_##CmdName##_Global(); /**< Maps UI command info CmdName with the specified UI command list. */\
-	const FUIAction CmdName##_Custom(); /**< UI action for CmdName command. */\
-protected:\
-	void CmdName##_Execute(); /**< Handles FExecuteAction for CmdName. */\
-	bool CmdName##_CanExecute() const; /**< Handles FCanExecuteAction for CmdName. */\
-	ECheckBoxState CmdName##_GetCheckState() const; /**< Handles FGetActionCheckState for CmdName. */
-
-	DECLARE_TOGGLE_COMMAND(ToggleObjectTableTreeViewVisibility)
-	DECLARE_TOGGLE_COMMAND(ToggleObjectDetailsViewVisibility)
-#undef DECLARE_TOGGLE_COMMAND
+	INSIGHTS_DECLARE_TOGGLE_COMMAND(ToggleObjectTableTreeViewVisibility)
+	INSIGHTS_DECLARE_TOGGLE_COMMAND(ToggleObjectDetailsViewVisibility)
 
 	//////////////////////////////////////////////////
 
diff --git a/Engine/Source/Developer/TraceInsights/Private/Insights/ObjectProfiler/ObjectProfilerManager.cpp b/Engine/Source/Developer/TraceInsights/Private/Insights/ObjectProfiler/ObjectProfilerManager.cpp
--- a/Engine/Source/Developer/TraceInsights/Private/Insights/ObjectProfiler/ObjectProfilerManager.cpp
+++ b/Engine/Source/Developer/TraceInsights/Private/Insights/ObjectProfiler/ObjectProfilerManager.cpp
@@ -142,8 +142,8 @@ FObjectProfilerManager::~FObjectProfilerManager()
 
 void FObjectProfilerManager::BindCommands()
 {
-	ActionManager.Map_ToggleObjectTableTreeViewVisibility_Global();
-	ActionManager.Map_ToggleObjectDetailsViewVisibility_Global();
+	ActionManager.Map_ToggleObjectTableTreeViewVisibility();
+	ActionManager.Map_ToggleObjectDetailsViewVisibility();
 }
 
 ////////////////////////////////////////////////////////////////////////////////////////////////////
diff --git a/Engine/Source/Developer/TraceInsights/Private/Insights/TimingProfiler/TimingProfilerCommands.cpp b/Engine/Source/Developer/TraceInsights/Private/Insights/TimingProfiler/TimingProfilerCommands.cpp
--- a/Engine/Source/Developer/TraceInsights/Private/Insights/TimingProfiler/TimingProfilerCommands.cpp
+++ b/Engine/Source/Developer/TraceInsights/Private/Insights/TimingProfiler/TimingProfilerCommands.cpp
@@ -4,7 +4,6 @@
 
 #include "DesktopPlatformModule.h"
 #include "Framework/Application/SlateApplication.h"
-#include "Framework/MultiBox/MultiBoxBuilder.h"
 
 // TraceInsights
 #include "Insights/InsightsManager.h"
@@ -18,23 +17,6 @@
 namespace UE::Insights::TimingProfiler
 {
 
-////////////////////////////////////////////////////////////////////////////////////////////////////
-// FTimingProfilerMenuBuilder
-////////////////////////////////////////////////////////////////////////////////////////////////////
-
-void FTimingProfilerMenuBuilder::AddMenuEntry(FMenuBuilder& MenuBuilder, const TSharedPtr<FUICommandInfo>& UICommandInfo, const FUIAction& UIAction)
-{
-	MenuBuilder.AddMenuEntry
-	(
-		UICommandInfo->GetLabel(),
-		UICommandInfo->GetDescription(),
-		UICommandInfo->GetIcon(),
-		UIAction,
-		NAME_None,
-		UICommandInfo->GetUserInterfaceType()
-	);
-}
-
 ////////////////////////////////////////////////////////////////////////////////////////////////////
 // FTimingProfilerCommands
 ////////////////////////////////////////////////////////////////////////////////////////////////////
@@ -54,14 +36,53 @@ FTimingProfilerCommands::FTimingProfilerCommands()
 UE_DISABLE_OPTIMIZATION_SHIP
 void FTimingProfilerCommands::RegisterCommands()
 {
-	UI_COMMAND(ToggleFramesTrackVisibility, "Frames", "Toggles the visibility of the Frames track.", EUserInterfaceActionType::ToggleButton, FInputChord());
-	UI_COMMAND(ToggleTimingViewVisibility, "Timing", "Toggles the visibility of the main Timing view.", EUserInterfaceActionType::ToggleButton, FInputChord());
-	UI_COMMAND(ToggleTimersViewVisibility, "Timers", "Toggles the visibility of the Timers view.", EUserInterfaceActionType::ToggleButton, FInputChord());
-	UI_COMMAND(ToggleCallersTreeViewVisibility, "Callers", "Toggles the visibility of the Callers tree view.", EUserInterfaceActionType::ToggleButton, FInputChord());
-	UI_COMMAND(ToggleCalleesTreeViewVisibility, "Callees", "Toggles the visibility of the Callees tree view.", EUserInterfaceActionType::ToggleButton,FInputChord());
-	UI_COMMAND(ToggleStatsCountersViewVisibility, "Counters", "Toggles the visibility of the Counters view.", EUserInterfaceActionType::ToggleButton, FInputChord());
-	UI_COMMAND(ToggleLogViewVisibility, "Log", "Toggles the visibility of the Log view.", EUserInterfaceActionType::ToggleButton,FInputChord());
-	UI_COMMAND(ToggleModulesViewVisibility, "Modules", "Toggles the visibility of the Modules view.", EUserInterfaceActionType::ToggleButton, FInputChord());
+	UI_COMMAND(ToggleFramesTrackVisibility,
+		"Frames",
+		"Toggles the visibility of the Frames track.",
+		EUserInterfaceActionType::ToggleButton,
+		FInputChord());
+
+	UI_COMMAND(ToggleTimingViewVisibility,
+		"Timing",
+		"Toggles the visibility of the main Timing view.",
+		EUserInterfaceActionType::ToggleButton,
+		FInputChord());
+
+	UI_COMMAND(ToggleTimersViewVisibility,
+		"Timers",
+		"Toggles the visibility of the Timers view.",
+		EUserInterfaceActionType::ToggleButton,
+		FInputChord());
+
+	UI_COMMAND(ToggleCallersTreeViewVisibility,
+		"Callers",
+		"Toggles the visibility of the Callers tree view.",
+		EUserInterfaceActionType::ToggleButton,
+		FInputChord());
+
+	UI_COMMAND(ToggleCalleesTreeViewVisibility,
+		"Callees",
+		"Toggles the visibility of the Callees tree view.",
+		EUserInterfaceActionType::ToggleButton,
+		FInputChord());
+
+	UI_COMMAND(ToggleStatsCountersViewVisibility,
+		"Counters",
+		"Toggles the visibility of the Counters view.",
+		EUserInterfaceActionType::ToggleButton,
+		FInputChord());
+
+	UI_COMMAND(ToggleLogViewVisibility,
+		"Log",
+		"Toggles the visibility of the Log view.",
+		EUserInterfaceActionType::ToggleButton,
+		FInputChord());
+
+	UI_COMMAND(ToggleModulesViewVisibility,
+		"Modules",
+		"Toggles the visibility of the Modules view.",
+		EUserInterfaceActionType::ToggleButton,
+		FInputChord());
 }
 UE_ENABLE_OPTIMIZATION_SHIP
 
@@ -126,49 +147,14 @@ UE_ENABLE_OPTIMIZATION_SHIP
 // Toggle Commands
 ////////////////////////////////////////////////////////////////////////////////////////////////////
 
-#define IMPLEMENT_TOGGLE_COMMAND(CmdName, IsEnabled, SetIsEnabled) \
-	\
-	void FTimingProfilerActionManager::Map_##CmdName##_Global()\
-	{\
-		This->CommandList->MapAction(This->GetCommands().CmdName, CmdName##_Custom());\
-	}\
-	\
-	const FUIAction FTimingProfilerActionManager::CmdName##_Custom() \
-	{\
-		FUIAction UIAction;\
-		UIAction.ExecuteAction = FExecuteAction::CreateRaw(this, &FTimingProfilerActionManager::CmdName##_Execute);\
-		UIAction.CanExecuteAction = FCanExecuteAction::CreateRaw(this, &FTimingProfilerActionManager::CmdName##_CanExecute);\
-		UIAction.GetActionCheckState = FGetActionCheckState::CreateRaw(this, &FTimingProfilerActionManager::CmdName##_GetCheckState);\
-		return UIAction;\
-	}\
-	\
-	void FTimingProfilerActionManager::CmdName##_Execute()\
-	{\
-		const bool b##IsEnabled = !This->IsEnabled();\
-		This->SetIsEnabled(b##IsEnabled);\
-	}\
-	\
-	bool FTimingProfilerActionManager::CmdName##_CanExecute() const\
-	{\
-		return FInsightsManager::Get()->GetSession().IsValid();\
-	}\
-	\
-	ECheckBoxState FTimingProfilerActionManager::CmdName##_GetCheckState() const\
-	{\
-		const bool b##IsEnabled = This->IsEnabled();\
-		return b##IsEnabled ? ECheckBoxState::Checked : ECheckBoxState::Unchecked;\
-	}
-
-IMPLEMENT_TOGGLE_COMMAND(ToggleFramesTrackVisibility, IsFramesTrackVisible, ShowHideFramesTrack)
-IMPLEMENT_TOGGLE_COMMAND(ToggleTimingViewVisibility, IsTimingViewVisible, ShowHideTimingView)
-IMPLEMENT_TOGGLE_COMMAND(ToggleTimersViewVisibility, IsTimersViewVisible, ShowHideTimersView)
-IMPLEMENT_TOGGLE_COMMAND(ToggleCallersTreeViewVisibility, IsCallersTreeViewVisible, ShowHideCallersTreeView)
-IMPLEMENT_TOGGLE_COMMAND(ToggleCalleesTreeViewVisibility, IsCalleesTreeViewVisible, ShowHideCalleesTreeView)
-IMPLEMENT_TOGGLE_COMMAND(ToggleStatsCountersViewVisibility, IsStatsCountersViewVisible, ShowHideStatsCountersView)
-IMPLEMENT_TOGGLE_COMMAND(ToggleLogViewVisibility, IsLogViewVisible, ShowHideLogView)
-IMPLEMENT_TOGGLE_COMMAND(ToggleModulesViewVisibility, IsModulesViewVisible, ShowHideModulesView)
-
-#undef IMPLEMENT_TOGGLE_COMMAND
+INSIGHTS_IMPLEMENT_TOGGLE_COMMAND(FTimingProfilerActionManager, This, ToggleFramesTrackVisibility, IsFramesTrackVisible, ShowHideFramesTrack)
+INSIGHTS_IMPLEMENT_TOGGLE_COMMAND(FTimingProfilerActionManager, This, ToggleTimingViewVisibility, IsTimingViewVisible, ShowHideTimingView)
+INSIGHTS_IMPLEMENT_TOGGLE_COMMAND(FTimingProfilerActionManager, This, ToggleTimersViewVisibility, IsTimersViewVisible, ShowHideTimersView)
+INSIGHTS_IMPLEMENT_TOGGLE_COMMAND(FTimingProfilerActionManager, This, ToggleCallersTreeViewVisibility, IsCallersTreeViewVisible, ShowHideCallersTreeView)
+INSIGHTS_IMPLEMENT_TOGGLE_COMMAND(FTimingProfilerActionManager, This, ToggleCalleesTreeViewVisibility, IsCalleesTreeViewVisible, ShowHideCalleesTreeView)
+INSIGHTS_IMPLEMENT_TOGGLE_COMMAND(FTimingProfilerActionManager, This, ToggleStatsCountersViewVisibility, IsStatsCountersViewVisible, ShowHideStatsCountersView)
+INSIGHTS_IMPLEMENT_TOGGLE_COMMAND(FTimingProfilerActionManager, This, ToggleLogViewVisibility, IsLogViewVisible, ShowHideLogView)
+INSIGHTS_IMPLEMENT_TOGGLE_COMMAND(FTimingProfilerActionManager, This, ToggleModulesViewVisibility, IsModulesViewVisible, ShowHideModulesView)
 
 ////////////////////////////////////////////////////////////////////////////////////////////////////
 
diff --git a/Engine/Source/Developer/TraceInsights/Private/Insights/TimingProfiler/TimingProfilerCommands.h b/Engine/Source/Developer/TraceInsights/Private/Insights/TimingProfiler/TimingProfilerCommands.h
--- a/Engine/Source/Developer/TraceInsights/Private/Insights/TimingProfiler/TimingProfilerCommands.h
+++ b/Engine/Source/Developer/TraceInsights/Private/Insights/TimingProfiler/TimingProfilerCommands.h
@@ -7,7 +7,8 @@
 #include "Framework/Commands/UIAction.h"
 #include "Framework/Commands/Commands.h"
 
-class FMenuBuilder;
+// TraceInsights
+#include "Insights/Common/CommandUtils.h"
 
 namespace UE::Insights::TimingProfiler
 {
@@ -24,35 +25,28 @@ public:
 	virtual void RegisterCommands() override;
 
 public:
-	//////////////////////////////////////////////////
-	// Global commands need to implement following method:
-	//     void Map_<CommandName>_Global();
-	// Custom commands needs to implement also the following method:
-	//     const FUIAction <CommandName>_Custom(...) const;
-	//////////////////////////////////////////////////
-
-	/** Toggles visibility for the Frames Track. Global and custom command. */
+	/** Toggles visibility for the Frames Track. */
 	TSharedPtr<FUICommandInfo> ToggleFramesTrackVisibility;
 
-	/** Toggles visibility for the Timing View. Global and custom command. */
+	/** Toggles visibility for the Timing View. */
 	TSharedPtr<FUICommandInfo> ToggleTimingViewVisibility;
 
-	/** Toggles visibility for the Timers View. Global and custom command. */
+	/** Toggles visibility for the Timers View. */
 	TSharedPtr<FUICommandInfo> ToggleTimersViewVisibility;
 
-	/** Toggles visibility for the Callers Tree View. Global and custom command. */
+	/** Toggles visibility for the Callers Tree View. */
 	TSharedPtr<FUICommandInfo> ToggleCallersTreeViewVisibility;
 
-	/** Toggles visibility for the Callees Tree View. Global and custom command. */
+	/** Toggles visibility for the Callees Tree View. */
 	TSharedPtr<FUICommandInfo> ToggleCalleesTreeViewVisibility;
 
-	/** Toggles visibility for the Stats Counters View. Global and custom command. */
+	/** Toggles visibility for the Stats Counters View. */
 	TSharedPtr<FUICommandInfo> ToggleStatsCountersViewVisibility;
 
-	/** Toggles visibility for the Log View. Global and custom command. */
+	/** Toggles visibility for the Log View. */
 	TSharedPtr<FUICommandInfo> ToggleLogViewVisibility;
 
-	/** Toggles visibility for the Modules view. Global and custom command. */
+	/** Toggles visibility for the Modules view. */
 	TSharedPtr<FUICommandInfo> ToggleModulesViewVisibility;
 };
 
@@ -85,28 +79,6 @@ public:
 	TSharedPtr<FUICommandInfo> TrimAndSaveAs;
 };
 
-////////////////////////////////////////////////////////////////////////////////////////////////////
-/**
- * Menu builder. Helper class for adding a customized menu entry using the global UI command info.
- */
-class FTimingProfilerMenuBuilder
-{
-public:
-	/**
-	 * Helper method for adding a customized menu entry using the global UI command info.
-	 * FUICommandInfo cannot be executed with custom parameters, so we need to create a custom FUIAction,
-	 * but sometime we have global and local version for the UI command, so reuse data from the global UI command info.
-	 * Ex:
-	 *     SessionInstance_ToggleCapture          - Global version will toggle capture process for all active session instances
-	 *     SessionInstance_ToggleCapture_OneParam - Local version will toggle capture process only for the specified session instance
-	 *
-	 * @param MenuBuilder The menu to add items to
-	 * @param FUICommandInfo A shared pointer to the UI command info
-	 * @param UIAction Customized version of the UI command info stored in an UI action
-	 */
-	static void AddMenuEntry(FMenuBuilder& MenuBuilder, const TSharedPtr<FUICommandInfo>& UICommandInfo, const FUIAction& UIAction);
-};
-
 ////////////////////////////////////////////////////////////////////////////////////////////////////
 /**
  * Class that provides helper functions for the commands to avoid cluttering profiler manager with many small functions.
@@ -125,34 +97,14 @@ private:
 	//////////////////////////////////////////////////
 	// Toggle Commands
 
-#define DECLARE_TOGGLE_COMMAND(CmdName)\
-public:\
-	void Map_##CmdName##_Global(); /**< Maps UI command info CmdName with the specified UI command list. */\
-	const FUIAction CmdName##_Custom(); /**< UI action for CmdName command. */\
-protected:\
-	void CmdName##_Execute(); /**< Handles FExecuteAction for CmdName. */\
-	bool CmdName##_CanExecute() const; /**< Handles FCanExecuteAction for CmdName. */\
-	ECheckBoxState CmdName##_GetCheckState() const; /**< Handles FGetActionCheckState for CmdName. */
-
-	DECLARE_TOGGLE_COMMAND(ToggleFramesTrackVisibility)
-	DECLARE_TOGGLE_COMMAND(ToggleTimingViewVisibility)
-	DECLARE_TOGGLE_COMMAND(ToggleTimersViewVisibility)
-	DECLARE_TOGGLE_COMMAND(ToggleCallersTreeViewVisibility)
-	DECLARE_TOGGLE_COMMAND(ToggleCalleesTreeViewVisibility)
-	DECLARE_TOGGLE_COMMAND(ToggleStatsCountersViewVisibility)
-	DECLARE_TOGGLE_COMMAND(ToggleLogViewVisibility)
-	DECLARE_TOGGLE_COMMAND(ToggleModulesViewVisibility)
-#undef DECLARE_TOGGLE_COMMAND
-
-	//////////////////////////////////////////////////
-	// OpenSettings
-
-public:
-	void Map_OpenSettings_Global(); /**< Maps UI command info OpenSettings with the specified UI command list. */
-	const FUIAction OpenSettings_Custom() const; /**< UI action for OpenSettings command. */
-protected:
-	void OpenSettings_Execute(); /**< Handles FExecuteAction for OpenSettings. */
-	bool OpenSettings_CanExecute() const; /**< Handles FCanExecuteAction for OpenSettings. */
+	INSIGHTS_DECLARE_TOGGLE_COMMAND(ToggleFramesTrackVisibility)
+	INSIGHTS_DECLARE_TOGGLE_COMMAND(ToggleTimingViewVisibility)
+	INSIGHTS_DECLARE_TOGGLE_COMMAND(ToggleTimersViewVisibility)
+	INSIGHTS_DECLARE_TOGGLE_COMMAND(ToggleCallersTreeViewVisibility)
+	INSIGHTS_DECLARE_TOGGLE_COMMAND(ToggleCalleesTreeViewVisibility)
+	INSIGHTS_DECLARE_TOGGLE_COMMAND(ToggleStatsCountersViewVisibility)
+	INSIGHTS_DECLARE_TOGGLE_COMMAND(ToggleLogViewVisibility)
+	INSIGHTS_DECLARE_TOGGLE_COMMAND(ToggleModulesViewVisibility)
 
 	//////////////////////////////////////////////////
 
diff --git a/Engine/Source/Developer/TraceInsights/Private/Insights/TimingProfiler/TimingProfilerManager.cpp b/Engine/Source/Developer/TraceInsights/Private/Insights/TimingProfiler/TimingProfilerManager.cpp
--- a/Engine/Source/Developer/TraceInsights/Private/Insights/TimingProfiler/TimingProfilerManager.cpp
+++ b/Engine/Source/Developer/TraceInsights/Private/Insights/TimingProfiler/TimingProfilerManager.cpp
@@ -158,14 +158,14 @@ FTimingProfilerManager::~FTimingProfilerManager()
 
 void FTimingProfilerManager::BindCommands()
 {
-	ActionManager.Map_ToggleFramesTrackVisibility_Global();
-	ActionManager.Map_ToggleTimingViewVisibility_Global();
-	ActionManager.Map_ToggleTimersViewVisibility_Global();
-	ActionManager.Map_ToggleCallersTreeViewVisibility_Global();
-	ActionManager.Map_ToggleCalleesTreeViewVisibility_Global();
-	ActionManager.Map_ToggleStatsCountersViewVisibility_Global();
-	ActionManager.Map_ToggleLogViewVisibility_Global();
-	ActionManager.Map_ToggleModulesViewVisibility_Global();
+	ActionManager.Map_ToggleFramesTrackVisibility();
+	ActionManager.Map_ToggleTimingViewVisibility();
+	ActionManager.Map_ToggleTimersViewVisibility();
+	ActionManager.Map_ToggleCallersTreeViewVisibility();
+	ActionManager.Map_ToggleCalleesTreeViewVisibility();
+	ActionManager.Map_ToggleStatsCountersViewVisibility();
+	ActionManager.Map_ToggleLogViewVisibility();
+	ActionManager.Map_ToggleModulesViewVisibility();
 }
 
 ////////////////////////////////////////////////////////////////////////////////////////////////////
