[Insights] Added InvalidObjectId to FObjectNode.

#rnx
#Synced-CL 52677305

diff --git a/Engine/Source/Developer/TraceInsights/Private/Insights/ObjectProfiler/ViewModels/ObjectGroupingByObjectName.cpp b/Engine/Source/Developer/TraceInsights/Private/Insights/ObjectProfiler/ViewModels/ObjectGroupingByObjectName.cpp
--- a/Engine/Source/Developer/TraceInsights/Private/Insights/ObjectProfiler/ViewModels/ObjectGroupingByObjectName.cpp
+++ b/Engine/Source/Developer/TraceInsights/Private/Insights/ObjectProfiler/ViewModels/ObjectGroupingByObjectName.cpp
@@ -61,7 +61,7 @@ void FObjectGroupingByObjectName::GroupNodes(
 		}
 
 		FObjectNode& ObjNode = NodePtr->As<FObjectNode>();
-		if (ObjNode.GetObjectId() == uint32(-1))
+		if (ObjNode.GetObjectId() == FObjectNode::InvalidObjectId)
 		{
 			ParentGroup.AddChildAndSetParent(NodePtr);
 			continue;
diff --git a/Engine/Source/Developer/TraceInsights/Private/Insights/ObjectProfiler/ViewModels/ObjectGroupingByOuter.cpp b/Engine/Source/Developer/TraceInsights/Private/Insights/ObjectProfiler/ViewModels/ObjectGroupingByOuter.cpp
--- a/Engine/Source/Developer/TraceInsights/Private/Insights/ObjectProfiler/ViewModels/ObjectGroupingByOuter.cpp
+++ b/Engine/Source/Developer/TraceInsights/Private/Insights/ObjectProfiler/ViewModels/ObjectGroupingByOuter.cpp
@@ -61,7 +61,7 @@ void FObjectGroupingByOuter::GroupNodes(
 		}
 
 		FObjectNode& ObjNode = NodePtr->As<FObjectNode>();
-		if (ObjNode.GetObjectId() == uint32(-1))
+		if (ObjNode.GetObjectId() == FObjectNode::InvalidObjectId)
 		{
 			ParentGroup.AddChildAndSetParent(NodePtr);
 			continue;
diff --git a/Engine/Source/Developer/TraceInsights/Private/Insights/ObjectProfiler/ViewModels/ObjectNode.cpp b/Engine/Source/Developer/TraceInsights/Private/Insights/ObjectProfiler/ViewModels/ObjectNode.cpp
--- a/Engine/Source/Developer/TraceInsights/Private/Insights/ObjectProfiler/ViewModels/ObjectNode.cpp
+++ b/Engine/Source/Developer/TraceInsights/Private/Insights/ObjectProfiler/ViewModels/ObjectNode.cpp
@@ -112,7 +112,7 @@ FLinearColor FObjectNode::GetIconColor() const
 
 FLinearColor FObjectNode::GetColor() const
 {
-	if (ObjectId == uint32(-1))
+	if (ObjectId == InvalidObjectId)
 	{
 		return GetDefaultIconColor(true);
 	}
diff --git a/Engine/Source/Developer/TraceInsights/Private/Insights/ObjectProfiler/ViewModels/ObjectNode.h b/Engine/Source/Developer/TraceInsights/Private/Insights/ObjectProfiler/ViewModels/ObjectNode.h
--- a/Engine/Source/Developer/TraceInsights/Private/Insights/ObjectProfiler/ViewModels/ObjectNode.h
+++ b/Engine/Source/Developer/TraceInsights/Private/Insights/ObjectProfiler/ViewModels/ObjectNode.h
@@ -125,11 +125,14 @@ protected:
 	void SetClass(TWeakPtr<FObjectNode> InClass) { Class = InClass; }
 	void SetOuter(TWeakPtr<FObjectNode> InOuter) { Outer = InOuter; }
 
+public:
+	static constexpr uint32 InvalidObjectId = uint32(-1);
+
 private:
-	uint32 ObjectId = uint32(-1);
+	uint32 ObjectId = InvalidObjectId;
 
-	uint32 ClassId = uint32(-1);
-	uint32 OuterId = uint32(-1);
+	uint32 ClassId = InvalidObjectId;
+	uint32 OuterId = InvalidObjectId;
 
 	uint32 ObjectFlags = 0;
 
