[Insights] SObjectProfilerWindow: Added API to access the SSegmentedBarGraph widget and the segments array. Also disabled the code for adding debug segments.

#rnx
#Synced-CL 52677305

diff --git a/Engine/Source/Developer/TraceInsights/Private/Insights/ObjectProfiler/ViewModels/ObjectNode.h b/Engine/Source/Developer/TraceInsights/Private/Insights/ObjectProfiler/ViewModels/ObjectNode.h
--- a/Engine/Source/Developer/TraceInsights/Private/Insights/ObjectProfiler/ViewModels/ObjectNode.h
+++ b/Engine/Source/Developer/TraceInsights/Private/Insights/ObjectProfiler/ViewModels/ObjectNode.h
@@ -198,8 +198,8 @@ public:
 	virtual int32 GetStructureSize() const override { return StructureSize; }
 
 private:
-	uint32 SuperId = uint32(-1);
-	uint32 InheritanceSuperId = uint32(-1);
+	uint32 SuperId = InvalidObjectId;
+	uint32 InheritanceSuperId = InvalidObjectId;
 	int32 StructureSize = 0;
 };
 
diff --git a/Engine/Source/Developer/TraceInsights/Private/Insights/ObjectProfiler/Widgets/SObjectProfilerWindow.cpp b/Engine/Source/Developer/TraceInsights/Private/Insights/ObjectProfiler/Widgets/SObjectProfilerWindow.cpp
--- a/Engine/Source/Developer/TraceInsights/Private/Insights/ObjectProfiler/Widgets/SObjectProfilerWindow.cpp
+++ b/Engine/Source/Developer/TraceInsights/Private/Insights/ObjectProfiler/Widgets/SObjectProfilerWindow.cpp
@@ -283,7 +283,7 @@ TSharedRef<SWidget> SObjectProfilerWindow::CreateToolbar(TSharedPtr<FExtender> E
 	IUnrealInsightsModule& InsightsModule = FModuleManager::LoadModuleChecked<IUnrealInsightsModule>("TraceInsights");
 	const FInsightsMajorTabConfig& Config = InsightsModule.FindMajorTabConfig(ObjectProfilerTabId);
 
-#if 1 // debug code
+#if 0 // debug code
 	auto Segment1 = MakeShared<FCustomBarGraphSegment>();
 	Segment1->Size = 50.0;
 	Segment1->Text = LOCTEXT("Segment1", "Textures");
diff --git a/Engine/Source/Developer/TraceInsights/Private/Insights/ObjectProfiler/Widgets/SObjectProfilerWindow.h b/Engine/Source/Developer/TraceInsights/Private/Insights/ObjectProfiler/Widgets/SObjectProfilerWindow.h
--- a/Engine/Source/Developer/TraceInsights/Private/Insights/ObjectProfiler/Widgets/SObjectProfilerWindow.h
+++ b/Engine/Source/Developer/TraceInsights/Private/Insights/ObjectProfiler/Widgets/SObjectProfilerWindow.h
@@ -85,6 +85,10 @@ public:
 
 	void ShowHideSegmentedBarGraph(bool bOnOff) { /* TODO */ }
 
+	TArray<TSharedPtr<IBarGraphSegment>>& GetSegments() { return Segments; }
+	const TArray<TSharedPtr<IBarGraphSegment>>& GetSegments() const { return Segments; }
+	TSharedPtr<SSegmentedBarGraph> GetSegmentedBarGraph() const { return SegmentedBarGraph; }
+
 	const FName& GetLogListingName() const { return LogListingName; }
 
 	virtual TSharedRef<SWidget> GetWidget() override { return SharedThis(this); }
