[Insights] ObjectProfiler: Added "simplified mode" and "can show referencing actors" options for the object table.

#Synced-CL 52677305

diff --git a/Engine/Source/Developer/TraceInsights/Private/Insights/ObjectProfiler/ViewModels/ObjectTable.cpp b/Engine/Source/Developer/TraceInsights/Private/Insights/ObjectProfiler/ViewModels/ObjectTable.cpp
--- a/Engine/Source/Developer/TraceInsights/Private/Insights/ObjectProfiler/ViewModels/ObjectTable.cpp
+++ b/Engine/Source/Developer/TraceInsights/Private/Insights/ObjectProfiler/ViewModels/ObjectTable.cpp
@@ -55,7 +55,9 @@ const FName FObjectTableColumns::ReferencesColumnId("References");
 // FObjectTable
 ////////////////////////////////////////////////////////////////////////////////////////////////////
 
-FObjectTable::FObjectTable()
+FObjectTable::FObjectTable(bool bInIsSimplifiedMode, bool bInCanShowReferencingActors)
+	: bIsSimplifiedMode(bInIsSimplifiedMode)
+	, bCanShowReferencingActors(bInCanShowReferencingActors)
 {
 }
 
@@ -108,6 +110,7 @@ void FObjectTable::AddDefaultColumns()
 
 	//////////////////////////////////////////////////
 	// Hierarchy Column
+
 	{
 		const int32 HierarchyColumnIndex = -1;
 		const TCHAR* HierarchyColumnName = nullptr;
@@ -124,6 +127,8 @@ void FObjectTable::AddDefaultColumns()
 
 	//////////////////////////////////////////////////
 	// Object Count Column
+
+	if (!bIsSimplifiedMode)
 	{
 		TSharedRef<FTableColumn> ColumnRef = MakeShared<FTableColumn>(FObjectTableColumns::ObjectCountColumnId);
 		FTableColumn& Column = *ColumnRef;
@@ -170,6 +175,8 @@ void FObjectTable::AddDefaultColumns()
 	}
 	//////////////////////////////////////////////////
 	// Object Name Column
+
+	if (!bIsSimplifiedMode)
 	{
 		TSharedRef<FTableColumn> ColumnRef = MakeShared<FTableColumn>(FObjectTableColumns::ObjectNameColumnId);
 		FTableColumn& Column = *ColumnRef;
@@ -220,6 +227,8 @@ void FObjectTable::AddDefaultColumns()
 	}
 	//////////////////////////////////////////////////
 	// Object Id Column
+
+	if (!bIsSimplifiedMode)
 	{
 		TSharedRef<FTableColumn> ColumnRef = MakeShared<FTableColumn>(FObjectTableColumns::ObjectIdColumnId);
 		FTableColumn& Column = *ColumnRef;
@@ -275,6 +284,8 @@ void FObjectTable::AddDefaultColumns()
 	}
 	//////////////////////////////////////////////////
 	// Class Name Column
+
+	if (!bIsSimplifiedMode)
 	{
 		TSharedRef<FTableColumn> ColumnRef = MakeShared<FTableColumn>(FObjectTableColumns::ClassNameColumnId);
 		FTableColumn& Column = *ColumnRef;
@@ -326,6 +337,8 @@ void FObjectTable::AddDefaultColumns()
 	}
 	//////////////////////////////////////////////////
 	// Class Id Column
+
+	if (!bIsSimplifiedMode)
 	{
 		TSharedRef<FTableColumn> ColumnRef = MakeShared<FTableColumn>(FObjectTableColumns::ClassIdColumnId);
 		FTableColumn& Column = *ColumnRef;
@@ -380,6 +393,8 @@ void FObjectTable::AddDefaultColumns()
 	}
 	//////////////////////////////////////////////////
 	// Outer Name Column
+
+	if (!bIsSimplifiedMode)
 	{
 		TSharedRef<FTableColumn> ColumnRef = MakeShared<FTableColumn>(FObjectTableColumns::OuterNameColumnId);
 		FTableColumn& Column = *ColumnRef;
@@ -431,6 +446,8 @@ void FObjectTable::AddDefaultColumns()
 	}
 	//////////////////////////////////////////////////
 	// Outer Id Column
+
+	if (!bIsSimplifiedMode)
 	{
 		TSharedRef<FTableColumn> ColumnRef = MakeShared<FTableColumn>(FObjectTableColumns::OuterIdColumnId);
 		FTableColumn& Column = *ColumnRef;
@@ -485,6 +502,8 @@ void FObjectTable::AddDefaultColumns()
 	}
 	//////////////////////////////////////////////////
 	// Package Name Column
+
+	if (!bIsSimplifiedMode)
 	{
 		TSharedRef<FTableColumn> ColumnRef = MakeShared<FTableColumn>(FObjectTableColumns::PackageNameColumnId);
 		FTableColumn& Column = *ColumnRef;
@@ -536,6 +555,8 @@ void FObjectTable::AddDefaultColumns()
 	}
 	//////////////////////////////////////////////////
 	// Outer Package Object Id Column
+
+	if (!bIsSimplifiedMode)
 	{
 		TSharedRef<FTableColumn> ColumnRef = MakeShared<FTableColumn>(FObjectTableColumns::PackageIdColumnId);
 		FTableColumn& Column = *ColumnRef;
@@ -590,6 +611,8 @@ void FObjectTable::AddDefaultColumns()
 	}
 	//////////////////////////////////////////////////
 	// Package Unique Id Column (UPackage)
+
+	if (!bIsSimplifiedMode)
 	{
 		TSharedRef<FTableColumn> ColumnRef = MakeShared<FTableColumn>(FObjectTableColumns::PackageUniqueIdColumnId);
 		FTableColumn& Column = *ColumnRef;
@@ -651,6 +674,8 @@ void FObjectTable::AddDefaultColumns()
 	}
 	//////////////////////////////////////////////////
 	// Package Path Column (UPackage)
+
+	if (!bIsSimplifiedMode)
 	{
 		TSharedRef<FTableColumn> ColumnRef = MakeShared<FTableColumn>(FObjectTableColumns::PackagePathColumnId);
 		FTableColumn& Column = *ColumnRef;
@@ -700,6 +725,8 @@ void FObjectTable::AddDefaultColumns()
 	}
 	//////////////////////////////////////////////////
 	// Object Flags Column
+
+	if (!bIsSimplifiedMode)
 	{
 		TSharedRef<FTableColumn> ColumnRef = MakeShared<FTableColumn>(FObjectTableColumns::ObjectFlagsColumnId);
 		FTableColumn& Column = *ColumnRef;
@@ -761,6 +788,8 @@ void FObjectTable::AddDefaultColumns()
 	}
 	//////////////////////////////////////////////////
 	// Object Path Column
+
+	if (!bIsSimplifiedMode)
 	{
 		TSharedRef<FTableColumn> ColumnRef = MakeShared<FTableColumn>(FObjectTableColumns::ObjectPathColumnId);
 		FTableColumn& Column = *ColumnRef;
@@ -807,6 +836,8 @@ void FObjectTable::AddDefaultColumns()
 	}
 	//////////////////////////////////////////////////
 	// Verse Path Column
+
+	if (!bIsSimplifiedMode)
 	{
 		TSharedRef<FTableColumn> ColumnRef = MakeShared<FTableColumn>(FObjectTableColumns::VersePathColumnId);
 		FTableColumn& Column = *ColumnRef;
@@ -853,6 +884,8 @@ void FObjectTable::AddDefaultColumns()
 	}
 	//////////////////////////////////////////////////
 	// Super Id Column (UStruct and UClass)
+
+	if (!bIsSimplifiedMode)
 	{
 		TSharedRef<FTableColumn> ColumnRef = MakeShared<FTableColumn>(FObjectTableColumns::SuperIdColumnId);
 		FTableColumn& Column = *ColumnRef;
@@ -907,6 +940,8 @@ void FObjectTable::AddDefaultColumns()
 	}
 	//////////////////////////////////////////////////
 	// Inheritance Super Id Column (UStruct and UClass)
+
+	if (!bIsSimplifiedMode)
 	{
 		TSharedRef<FTableColumn> ColumnRef = MakeShared<FTableColumn>(FObjectTableColumns::InheritanceSuperIdColumnId);
 		FTableColumn& Column = *ColumnRef;
@@ -961,6 +996,8 @@ void FObjectTable::AddDefaultColumns()
 	}
 	//////////////////////////////////////////////////
 	// Structure Size Column
+
+	if (!bIsSimplifiedMode)
 	{
 		TSharedRef<FTableColumn> ColumnRef = MakeShared<FTableColumn>(FObjectTableColumns::StructureSizeColumnId);
 		FTableColumn& Column = *ColumnRef;
@@ -1013,6 +1050,7 @@ void FObjectTable::AddDefaultColumns()
 	}
 	//////////////////////////////////////////////////
 	// System Memory Size Column
+
 	{
 		TSharedRef<FTableColumn> ColumnRef = MakeShared<FTableColumn>(FObjectTableColumns::SystemMemorySizeColumnId);
 		FTableColumn& Column = *ColumnRef;
@@ -1023,9 +1061,17 @@ void FObjectTable::AddDefaultColumns()
 		Column.SetTitleName(LOCTEXT("SystemMemorySizeColumnTitle", "System Memory Size"));
 		Column.SetDescription(LOCTEXT("SystemMemorySizeColumnDesc", "The estimated system memory size (bytes) of the UObject instance"));
 
-		Column.SetFlags(ETableColumnFlags::ShouldBeVisible |
-						ETableColumnFlags::CanBeHidden |
-						ETableColumnFlags::CanBeFiltered);
+		if (bIsSimplifiedMode)
+		{
+			Column.SetFlags(ETableColumnFlags::CanBeHidden |
+							ETableColumnFlags::CanBeFiltered);
+		}
+		else
+		{
+			Column.SetFlags(ETableColumnFlags::ShouldBeVisible |
+							ETableColumnFlags::CanBeHidden |
+							ETableColumnFlags::CanBeFiltered);
+		}
 
 		Column.SetHorizontalAlignment(HAlign_Right);
 		Column.SetInitialWidth(100.0f);
@@ -1065,6 +1111,7 @@ void FObjectTable::AddDefaultColumns()
 	}
 	//////////////////////////////////////////////////
 	// Video Memory Size Column
+
 	{
 		TSharedRef<FTableColumn> ColumnRef = MakeShared<FTableColumn>(FObjectTableColumns::VideoMemorySizeColumnId);
 		FTableColumn& Column = *ColumnRef;
@@ -1075,9 +1122,17 @@ void FObjectTable::AddDefaultColumns()
 		Column.SetTitleName(LOCTEXT("VideoMemorySizeColumnTitle", "Video Memory Size"));
 		Column.SetDescription(LOCTEXT("VideoMemorySizeColumnDesc", "The estimated video memory size (bytes) of the UObject instance"));
 
-		Column.SetFlags(ETableColumnFlags::ShouldBeVisible |
-						ETableColumnFlags::CanBeHidden |
-						ETableColumnFlags::CanBeFiltered);
+		if (bIsSimplifiedMode)
+		{
+			Column.SetFlags(ETableColumnFlags::CanBeHidden |
+							ETableColumnFlags::CanBeFiltered);
+		}
+		else
+		{
+			Column.SetFlags(ETableColumnFlags::ShouldBeVisible |
+							ETableColumnFlags::CanBeHidden |
+							ETableColumnFlags::CanBeFiltered);
+		}
 
 		Column.SetHorizontalAlignment(HAlign_Right);
 		Column.SetInitialWidth(100.0f);
@@ -1117,6 +1172,7 @@ void FObjectTable::AddDefaultColumns()
 	}
 	//////////////////////////////////////////////////
 	// Estimated Size Column
+
 	{
 		TSharedRef<FTableColumn> ColumnRef = MakeShared<FTableColumn>(FObjectTableColumns::EstimatedSizeColumnId);
 		FTableColumn& Column = *ColumnRef;
@@ -1125,11 +1181,21 @@ void FObjectTable::AddDefaultColumns()
 
 		Column.SetShortName(LOCTEXT("EstimatedSizeColumnName", "Estimated Size"));
 		Column.SetTitleName(LOCTEXT("EstimatedSizeColumnTitle", "Estimated Size (MiB)"));
-		Column.SetDescription(LOCTEXT("EstimatedSizeColumnDesc", "The estimated memory size, in MiB"));
+		Column.SetDescription(LOCTEXT("EstimatedSizeColumnDesc", "The estimated memory size (system memory + video memory), in MiB"));
 
-		Column.SetFlags(ETableColumnFlags::CanBeHidden |
-						ETableColumnFlags::CanBeFiltered |
-						ETableColumnFlags::IsDynamic);
+		if (bIsSimplifiedMode)
+		{
+			Column.SetFlags(ETableColumnFlags::ShouldBeVisible |
+							ETableColumnFlags::CanBeHidden |
+							ETableColumnFlags::CanBeFiltered |
+							ETableColumnFlags::IsDynamic);
+		}
+		else
+		{
+			Column.SetFlags(ETableColumnFlags::CanBeHidden |
+							ETableColumnFlags::CanBeFiltered |
+							ETableColumnFlags::IsDynamic);
+		}
 
 		Column.SetHorizontalAlignment(HAlign_Right);
 		Column.SetInitialWidth(60.0f);
@@ -1205,6 +1271,7 @@ void FObjectTable::AddDefaultColumns()
 	}
 	//////////////////////////////////////////////////
 	// Impact Column
+
 	{
 		TSharedRef<FTableColumn> ColumnRef = MakeShared<FTableColumn>(FObjectTableColumns::ImpactColumnId);
 		FTableColumn& Column = *ColumnRef;
@@ -1215,9 +1282,19 @@ void FObjectTable::AddDefaultColumns()
 		Column.SetTitleName(LOCTEXT("ImpactColumnTitle", "Impact (%)"));
 		Column.SetDescription(LOCTEXT("ImpactColumnDesc", "The impact (percentage) of estimated memory size"));
 
-		Column.SetFlags(ETableColumnFlags::CanBeHidden |
-			ETableColumnFlags::CanBeFiltered |
-			ETableColumnFlags::IsDynamic);
+		if (bIsSimplifiedMode)
+		{
+			Column.SetFlags(ETableColumnFlags::ShouldBeVisible |
+							ETableColumnFlags::CanBeHidden |
+							ETableColumnFlags::CanBeFiltered |
+							ETableColumnFlags::IsDynamic);
+		}
+		else
+		{
+			Column.SetFlags(ETableColumnFlags::CanBeHidden |
+							ETableColumnFlags::CanBeFiltered |
+							ETableColumnFlags::IsDynamic);
+		}
 
 		Column.SetHorizontalAlignment(HAlign_Right);
 		Column.SetInitialWidth(60.0f);
@@ -1279,6 +1356,8 @@ void FObjectTable::AddDefaultColumns()
 	}
 	//////////////////////////////////////////////////
 	// Inclusive System Memory Size Column
+
+	if (!bIsSimplifiedMode)
 	{
 		TSharedRef<FTableColumn> ColumnRef = MakeShared<FTableColumn>(FObjectTableColumns::InclSystemMemorySizeColumnId);
 		FTableColumn& Column = *ColumnRef;
@@ -1330,6 +1409,8 @@ void FObjectTable::AddDefaultColumns()
 	}
 	//////////////////////////////////////////////////
 	// Inclusive Video Memory Size Column
+
+	if (!bIsSimplifiedMode)
 	{
 		TSharedRef<FTableColumn> ColumnRef = MakeShared<FTableColumn>(FObjectTableColumns::InclVideoMemorySizeColumnId);
 		FTableColumn& Column = *ColumnRef;
@@ -1381,6 +1462,8 @@ void FObjectTable::AddDefaultColumns()
 	}
 	//////////////////////////////////////////////////
 	// Inclusive Estimated Size Column
+
+	if (!bIsSimplifiedMode)
 	{
 		TSharedRef<FTableColumn> ColumnRef = MakeShared<FTableColumn>(FObjectTableColumns::InclEstimatedSizeColumnId);
 		FTableColumn& Column = *ColumnRef;
@@ -1469,6 +1552,8 @@ void FObjectTable::AddDefaultColumns()
 	}
 	//////////////////////////////////////////////////
 	// References Column
+
+	if (bCanShowReferencingActors)
 	{
 		TSharedRef<FTableColumn> ColumnRef = MakeShared<FTableColumn>(FObjectTableColumns::ReferencesColumnId);
 		FTableColumn& Column = *ColumnRef;
@@ -1479,8 +1564,17 @@ void FObjectTable::AddDefaultColumns()
 		Column.SetTitleName(LOCTEXT("ReferencesColumnTitle", "References"));
 		Column.SetDescription(LOCTEXT("ReferencesColumnDesc", "The number of asset references"));
 
-		Column.SetFlags(ETableColumnFlags::CanBeHidden |
-						ETableColumnFlags::CanBeFiltered);
+		if (bIsSimplifiedMode)
+		{
+			Column.SetFlags(ETableColumnFlags::ShouldBeVisible |
+							ETableColumnFlags::CanBeHidden |
+							ETableColumnFlags::CanBeFiltered);
+		}
+		else
+		{
+			Column.SetFlags(ETableColumnFlags::CanBeHidden |
+							ETableColumnFlags::CanBeFiltered);
+		}
 
 		Column.SetHorizontalAlignment(HAlign_Right);
 		Column.SetInitialWidth(60.0f);
diff --git a/Engine/Source/Developer/TraceInsights/Private/Insights/ObjectProfiler/ViewModels/ObjectTable.h b/Engine/Source/Developer/TraceInsights/Private/Insights/ObjectProfiler/ViewModels/ObjectTable.h
--- a/Engine/Source/Developer/TraceInsights/Private/Insights/ObjectProfiler/ViewModels/ObjectTable.h
+++ b/Engine/Source/Developer/TraceInsights/Private/Insights/ObjectProfiler/ViewModels/ObjectTable.h
@@ -54,13 +54,17 @@ struct FObjectTableColumns
 class FObjectTable : public FTable
 {
 public:
-	FObjectTable();
+	FObjectTable(bool bInIsSimplifiedMode = false, bool bInCanShowReferencingActors = true);
 	virtual ~FObjectTable();
 
 	virtual void Reset();
 
 private:
 	void AddDefaultColumns();
+
+private:
+	bool bIsSimplifiedMode = false;
+	bool bCanShowReferencingActors = false;
 };
 
 ////////////////////////////////////////////////////////////////////////////////////////////////////
